buildc script, add new samples for 4.1

diff --git a/Engine/Source/UE4Editor.Target.cs b/Engine/Source/UE4Editor.Target.cs
--- a/Engine/Source/UE4Editor.Target.cs
+++ b/Engine/Source/UE4Editor.Target.cs
@@ -76,6 +76,7 @@ public class UE4EditorTarget : TargetRules
         NonCodeProjectNames.Add("BlackJack", new List<UnrealTargetPlatform> { HostPlatform });
         NonCodeProjectNames.Add("Card", new List<UnrealTargetPlatform> { HostPlatform });
         NonCodeProjectNames.Add("Effects", new List<UnrealTargetPlatform> { HostPlatform });
+        NonCodeProjectNames.Add("GDC2014", new List<UnrealTargetPlatform> { HostPlatform });
 
         List<UnrealTargetPlatform> MobilePlats = null;
         if (HostPlatform == UnrealTargetPlatform.Mac)
