build script, stylised -> stylized

diff --git a/Engine/Source/UE4Editor.Target.cs b/Engine/Source/UE4Editor.Target.cs
--- a/Engine/Source/UE4Editor.Target.cs
+++ b/Engine/Source/UE4Editor.Target.cs
@@ -69,7 +69,7 @@ public class UE4EditorTarget : TargetRules
         NonCodeProjectNames.Add("Reflections", new List<UnrealTargetPlatform> { HostPlatform });
         NonCodeProjectNames.Add("ElementalVR", new List<UnrealTargetPlatform> { HostPlatform });
         NonCodeProjectNames.Add("VersusVR", new List<UnrealTargetPlatform> { HostPlatform });
-        NonCodeProjectNames.Add("Stylised", new List<UnrealTargetPlatform> { HostPlatform });
+        NonCodeProjectNames.Add("Stylized", new List<UnrealTargetPlatform> { HostPlatform });
         NonCodeProjectNames.Add("Landscape", new List<UnrealTargetPlatform> { HostPlatform });
         NonCodeProjectNames.Add("Matinee", new List<UnrealTargetPlatform> { HostPlatform });
         NonCodeProjectNames.Add("RealisticRendering", new List<UnrealTargetPlatform> { HostPlatform });
