build script, adjust samples for rocket

diff --git a/Engine/Source/UE4Editor.Target.cs b/Engine/Source/UE4Editor.Target.cs
--- a/Engine/Source/UE4Editor.Target.cs
+++ b/Engine/Source/UE4Editor.Target.cs
@@ -73,10 +73,9 @@ public class UE4EditorTarget : TargetRules
         NonCodeProjectNames.Add("Landscape", new List<UnrealTargetPlatform> { HostPlatform });
         NonCodeProjectNames.Add("Matinee", new List<UnrealTargetPlatform> { HostPlatform });
         NonCodeProjectNames.Add("RealisticRendering", new List<UnrealTargetPlatform> { HostPlatform });
-        NonCodeProjectNames.Add("BlackJack", new List<UnrealTargetPlatform> { HostPlatform });
-        NonCodeProjectNames.Add("Card", new List<UnrealTargetPlatform> { HostPlatform });
         NonCodeProjectNames.Add("Effects", new List<UnrealTargetPlatform> { HostPlatform });
         NonCodeProjectNames.Add("GDC2014", new List<UnrealTargetPlatform> { HostPlatform });
+        NonCodeProjectNames.Add("ContentExamples", new List<UnrealTargetPlatform> { HostPlatform });
 
         List<UnrealTargetPlatform> MobilePlats = null;
         if (HostPlatform == UnrealTargetPlatform.Mac)
@@ -88,7 +87,8 @@ public class UE4EditorTarget : TargetRules
             MobilePlats = new List<UnrealTargetPlatform> { HostPlatform, UnrealTargetPlatform.Android };
         }
 
-        NonCodeProjectNames.Add("ContentExamples", MobilePlats);
+        NonCodeProjectNames.Add("BlackJack", MobilePlats);
+        NonCodeProjectNames.Add("Card", MobilePlats);
         NonCodeProjectNames.Add("TappyChicken", MobilePlats);
         NonCodeProjectNames.Add("SwingNinja", MobilePlats);
         NonCodeProjectNames.Add("Mobile", MobilePlats);
