build script, add tappy chicken android formal build

diff --git a/Engine/Source/Programs/AutomationTool/Scripts/GUBP.Automation.cs b/Engine/Source/Programs/AutomationTool/Scripts/GUBP.Automation.cs
--- a/Engine/Source/Programs/AutomationTool/Scripts/GUBP.Automation.cs
+++ b/Engine/Source/Programs/AutomationTool/Scripts/GUBP.Automation.cs
@@ -4061,7 +4061,7 @@ public class GUBP : BuildCommand
                     }
 
                     var TempNonCodeFormalBuilds = Target.Rules.GUBP_NonCodeFormalBuilds_BaseEditorTypeOnly();
-                    var HostMonos = GetMonolithicPlatformsForUProject(UnrealTargetPlatform.Mac, Branch.BaseEngineProject, true);
+                    var HostMonos = GetMonolithicPlatformsForUProject(HostPlatform, Branch.BaseEngineProject, true);
 
                     foreach (var Codeless in TempNonCodeFormalBuilds)
                     {
@@ -4072,7 +4072,7 @@ public class GUBP : BuildCommand
                             if (HostMonos.Contains(PlatPair.Key))
                             {
                                 NewPlatList.Add(PlatPair);
-                }
+                            }
                         }
                         if (NewPlatList.Count > 0)
                         {
diff --git a/Engine/Source/UE4Editor.Target.cs b/Engine/Source/UE4Editor.Target.cs
--- a/Engine/Source/UE4Editor.Target.cs
+++ b/Engine/Source/UE4Editor.Target.cs
@@ -106,7 +106,13 @@ public class UE4EditorTarget : TargetRules
     public override Dictionary<string, List<KeyValuePair<UnrealTargetPlatform, UnrealTargetConfiguration>>> GUBP_NonCodeFormalBuilds_BaseEditorTypeOnly()
     {
         var NonCodeProjectNames = new Dictionary<string, List<KeyValuePair<UnrealTargetPlatform, UnrealTargetConfiguration>>>();
-        NonCodeProjectNames.Add("TappyChicken", new List<KeyValuePair<UnrealTargetPlatform, UnrealTargetConfiguration>>{new KeyValuePair<UnrealTargetPlatform, UnrealTargetConfiguration>(UnrealTargetPlatform.IOS, UnrealTargetConfiguration.Shipping)});
+        NonCodeProjectNames.Add("TappyChicken", 
+            new List<KeyValuePair<UnrealTargetPlatform, UnrealTargetConfiguration>>
+            {
+                    new KeyValuePair<UnrealTargetPlatform, UnrealTargetConfiguration>(UnrealTargetPlatform.IOS, UnrealTargetConfiguration.Shipping),
+                    new KeyValuePair<UnrealTargetPlatform, UnrealTargetConfiguration>(UnrealTargetPlatform.Android, UnrealTargetConfiguration.Shipping)
+            }
+        );
         return NonCodeProjectNames;
     }
 }
