build script, update rocket samples and a few fixed for IOS on PC

diff --git a/Engine/Source/UE4Editor.Target.cs b/Engine/Source/UE4Editor.Target.cs
--- a/Engine/Source/UE4Editor.Target.cs
+++ b/Engine/Source/UE4Editor.Target.cs
@@ -76,6 +76,11 @@ public class UE4EditorTarget : TargetRules
         NonCodeProjectNames.Add("Effects", new List<UnrealTargetPlatform> { HostPlatform });
         NonCodeProjectNames.Add("GDC2014", new List<UnrealTargetPlatform> { HostPlatform });
         NonCodeProjectNames.Add("ContentExamples", new List<UnrealTargetPlatform> { HostPlatform });
+        NonCodeProjectNames.Add("Physics", new List<UnrealTargetPlatform> { HostPlatform });
+        NonCodeProjectNames.Add("TowerDefenseGame", new List<UnrealTargetPlatform> { HostPlatform });
+        NonCodeProjectNames.Add("Landscape_WorldMachine", new List<UnrealTargetPlatform> { HostPlatform });
+        NonCodeProjectNames.Add("MorphTargets", new List<UnrealTargetPlatform> { HostPlatform });
+        NonCodeProjectNames.Add("PostProcessMatinee", new List<UnrealTargetPlatform> { HostPlatform });
 
         List<UnrealTargetPlatform> MobilePlats = null;
         if (HostPlatform == UnrealTargetPlatform.Mac)
@@ -84,7 +89,7 @@ public class UE4EditorTarget : TargetRules
         }
         else
         {
-            MobilePlats = new List<UnrealTargetPlatform> { HostPlatform, UnrealTargetPlatform.Android };
+            MobilePlats = new List<UnrealTargetPlatform> { HostPlatform, UnrealTargetPlatform.Android, UnrealTargetPlatform.IOS };
         }
 
         NonCodeProjectNames.Add("BlackJack", MobilePlats);
@@ -100,6 +105,7 @@ public class UE4EditorTarget : TargetRules
         NonCodeProjectNames.Add("TP_SideScrollerBP", MobilePlats);
         NonCodeProjectNames.Add("TP_ThirdPersonBP", MobilePlats);
         NonCodeProjectNames.Add("TP_TopDownBP", MobilePlats);
+        NonCodeProjectNames.Add("TP_VehicleBP", MobilePlats);
 
         return NonCodeProjectNames;
     }
