build script, rename samples

diff --git a/Engine/Source/UE4Editor.Target.cs b/Engine/Source/UE4Editor.Target.cs
--- a/Engine/Source/UE4Editor.Target.cs
+++ b/Engine/Source/UE4Editor.Target.cs
@@ -62,21 +62,21 @@ public class UE4EditorTarget : TargetRules
     public override Dictionary<string, List<UnrealTargetPlatform>> GUBP_NonCodeProjects_BaseEditorTypeOnly(UnrealTargetPlatform HostPlatform)
     {
         var NonCodeProjectNames = new Dictionary<string, List<UnrealTargetPlatform>>();
-        NonCodeProjectNames.Add("Elemental", new List<UnrealTargetPlatform> { HostPlatform });
-        NonCodeProjectNames.Add("Infiltrator", new List<UnrealTargetPlatform> { HostPlatform });
+        NonCodeProjectNames.Add("ElementalDemo", new List<UnrealTargetPlatform> { HostPlatform });
+        NonCodeProjectNames.Add("InfiltratorDemo", new List<UnrealTargetPlatform> { HostPlatform });
         NonCodeProjectNames.Add("HoverShip", new List<UnrealTargetPlatform> { HostPlatform });
-        NonCodeProjectNames.Add("BlueprintExamples", new List<UnrealTargetPlatform> { HostPlatform });
-        NonCodeProjectNames.Add("Reflections", new List<UnrealTargetPlatform> { HostPlatform });
+        NonCodeProjectNames.Add("BlueprintOffice", new List<UnrealTargetPlatform> { HostPlatform });
+        NonCodeProjectNames.Add("ReflectionsSubway", new List<UnrealTargetPlatform> { HostPlatform });
         NonCodeProjectNames.Add("ElementalVR", new List<UnrealTargetPlatform> { HostPlatform });
-        NonCodeProjectNames.Add("VersusVR", new List<UnrealTargetPlatform> { HostPlatform });
+        NonCodeProjectNames.Add("CouchKnights", new List<UnrealTargetPlatform> { HostPlatform });
         NonCodeProjectNames.Add("Stylized", new List<UnrealTargetPlatform> { HostPlatform });
         NonCodeProjectNames.Add("Landscape", new List<UnrealTargetPlatform> { HostPlatform });
-        NonCodeProjectNames.Add("Matinee", new List<UnrealTargetPlatform> { HostPlatform });
+        NonCodeProjectNames.Add("MatineeFightScene", new List<UnrealTargetPlatform> { HostPlatform });
         NonCodeProjectNames.Add("RealisticRendering", new List<UnrealTargetPlatform> { HostPlatform });
-        NonCodeProjectNames.Add("Effects", new List<UnrealTargetPlatform> { HostPlatform });
+        NonCodeProjectNames.Add("EffectsCave", new List<UnrealTargetPlatform> { HostPlatform });
         NonCodeProjectNames.Add("GDC2014", new List<UnrealTargetPlatform> { HostPlatform });
         NonCodeProjectNames.Add("ContentExamples", new List<UnrealTargetPlatform> { HostPlatform });
-        NonCodeProjectNames.Add("Physics", new List<UnrealTargetPlatform> { HostPlatform });
+        NonCodeProjectNames.Add("PhysicsPirateShip", new List<UnrealTargetPlatform> { HostPlatform });
         NonCodeProjectNames.Add("TowerDefenseGame", new List<UnrealTargetPlatform> { HostPlatform });
         NonCodeProjectNames.Add("Landscape_WorldMachine", new List<UnrealTargetPlatform> { HostPlatform });
         NonCodeProjectNames.Add("MorphTargets", new List<UnrealTargetPlatform> { HostPlatform });
@@ -93,10 +93,10 @@ public class UE4EditorTarget : TargetRules
         }
 
         NonCodeProjectNames.Add("BlackJack", MobilePlats);
-        NonCodeProjectNames.Add("Card", MobilePlats);
+        NonCodeProjectNames.Add("MemoryGame", MobilePlats);
         NonCodeProjectNames.Add("TappyChicken", MobilePlats);
         NonCodeProjectNames.Add("SwingNinja", MobilePlats);
-        NonCodeProjectNames.Add("Mobile", MobilePlats);
+        NonCodeProjectNames.Add("MobileTemple", MobilePlats);
 
         NonCodeProjectNames.Add("StarterContent", MobilePlats);
         NonCodeProjectNames.Add("TP_FirstPersonBP", MobilePlats);
