build script, update samples

diff --git a/Engine/Source/UE4Editor.Target.cs b/Engine/Source/UE4Editor.Target.cs
--- a/Engine/Source/UE4Editor.Target.cs
+++ b/Engine/Source/UE4Editor.Target.cs
@@ -70,7 +70,7 @@ public class UE4EditorTarget : TargetRules
         NonCodeProjectNames.Add("ElementalVR", new List<UnrealTargetPlatform> { HostPlatform });
         NonCodeProjectNames.Add("CouchKnights", new List<UnrealTargetPlatform> { HostPlatform });
         NonCodeProjectNames.Add("Stylized", new List<UnrealTargetPlatform> { HostPlatform });
-        NonCodeProjectNames.Add("Landscape", new List<UnrealTargetPlatform> { HostPlatform });
+        //NonCodeProjectNames.Add("Landscape", new List<UnrealTargetPlatform> { HostPlatform });
         NonCodeProjectNames.Add("MatineeFightScene", new List<UnrealTargetPlatform> { HostPlatform });
         NonCodeProjectNames.Add("RealisticRendering", new List<UnrealTargetPlatform> { HostPlatform });
         NonCodeProjectNames.Add("EffectsCave", new List<UnrealTargetPlatform> { HostPlatform });
@@ -78,7 +78,7 @@ public class UE4EditorTarget : TargetRules
         NonCodeProjectNames.Add("ContentExamples", new List<UnrealTargetPlatform> { HostPlatform });
         NonCodeProjectNames.Add("PhysicsPirateShip", new List<UnrealTargetPlatform> { HostPlatform });
         NonCodeProjectNames.Add("TowerDefenseGame", new List<UnrealTargetPlatform> { HostPlatform });
-        NonCodeProjectNames.Add("Landscape_WorldMachine", new List<UnrealTargetPlatform> { HostPlatform });
+        NonCodeProjectNames.Add("LandscapeWM", new List<UnrealTargetPlatform> { HostPlatform });
         NonCodeProjectNames.Add("MorphTargets", new List<UnrealTargetPlatform> { HostPlatform });
         NonCodeProjectNames.Add("PostProcessMatinee", new List<UnrealTargetPlatform> { HostPlatform });
 
