build script, add sci fri hallways, trim samples for 4.1

diff --git a/Engine/Source/UE4Editor.Target.cs b/Engine/Source/UE4Editor.Target.cs
--- a/Engine/Source/UE4Editor.Target.cs
+++ b/Engine/Source/UE4Editor.Target.cs
@@ -81,6 +81,7 @@ public class UE4EditorTarget : TargetRules
         NonCodeProjectNames.Add("LandscapeWM", new List<UnrealTargetPlatform> { HostPlatform });
         NonCodeProjectNames.Add("MorphTargets", new List<UnrealTargetPlatform> { HostPlatform });
         NonCodeProjectNames.Add("PostProcessMatinee", new List<UnrealTargetPlatform> { HostPlatform });
+        NonCodeProjectNames.Add("SciFiHallway", new List<UnrealTargetPlatform> { HostPlatform });
 
         List<UnrealTargetPlatform> MobilePlats = null;
         if (HostPlatform == UnrealTargetPlatform.Mac)
