build script, update samples

diff --git a/Engine/Source/UE4Editor.Target.cs b/Engine/Source/UE4Editor.Target.cs
--- a/Engine/Source/UE4Editor.Target.cs
+++ b/Engine/Source/UE4Editor.Target.cs
@@ -78,7 +78,7 @@ public class UE4EditorTarget : TargetRules
         NonCodeProjectNames.Add("ContentExamples", new List<UnrealTargetPlatform> { HostPlatform });
         NonCodeProjectNames.Add("PhysicsPirateShip", new List<UnrealTargetPlatform> { HostPlatform });
         NonCodeProjectNames.Add("TowerDefenseGame", new List<UnrealTargetPlatform> { HostPlatform });
-        NonCodeProjectNames.Add("LandscapeWM", new List<UnrealTargetPlatform> { HostPlatform });
+        NonCodeProjectNames.Add("LandscapeMountains", new List<UnrealTargetPlatform> { HostPlatform });
         NonCodeProjectNames.Add("MorphTargets", new List<UnrealTargetPlatform> { HostPlatform });
         NonCodeProjectNames.Add("PostProcessMatinee", new List<UnrealTargetPlatform> { HostPlatform });
         NonCodeProjectNames.Add("SciFiHallway", new List<UnrealTargetPlatform> { HostPlatform });
