Build script: cook non-code projects for Linux.

This is needed so that build farm can populate DDC. Hopefully this will resolve (or at least alleviate) long cooking times problem that Engine-QA guys are seeing (TTP 331980, 332278).

diff --git a/Engine/Source/UE4Editor.Target.cs b/Engine/Source/UE4Editor.Target.cs
--- a/Engine/Source/UE4Editor.Target.cs
+++ b/Engine/Source/UE4Editor.Target.cs
@@ -62,26 +62,37 @@ public class UE4EditorTarget : TargetRules
     public override Dictionary<string, List<UnrealTargetPlatform>> GUBP_NonCodeProjects_BaseEditorTypeOnly(UnrealTargetPlatform HostPlatform)
     {
         var NonCodeProjectNames = new Dictionary<string, List<UnrealTargetPlatform>>();
-        NonCodeProjectNames.Add("ElementalDemo", new List<UnrealTargetPlatform> { HostPlatform });
-        NonCodeProjectNames.Add("InfiltratorDemo", new List<UnrealTargetPlatform> { HostPlatform });
-        NonCodeProjectNames.Add("HoverShip", new List<UnrealTargetPlatform> { HostPlatform });
-        NonCodeProjectNames.Add("BlueprintOffice", new List<UnrealTargetPlatform> { HostPlatform });
-        NonCodeProjectNames.Add("ReflectionsSubway", new List<UnrealTargetPlatform> { HostPlatform });
-        NonCodeProjectNames.Add("ElementalVR", new List<UnrealTargetPlatform> { HostPlatform });
-        NonCodeProjectNames.Add("CouchKnights", new List<UnrealTargetPlatform> { HostPlatform });
-        NonCodeProjectNames.Add("Stylized", new List<UnrealTargetPlatform> { HostPlatform });
+
+        List<UnrealTargetPlatform> DesktopPlats = null;
+        if (HostPlatform == UnrealTargetPlatform.Win64)
+        {
+            DesktopPlats = new List<UnrealTargetPlatform> { HostPlatform, UnrealTargetPlatform.Linux };
+        }
+        else
+        {
+            DesktopPlats = new List<UnrealTargetPlatform> { HostPlatform };
+        }
+
+        NonCodeProjectNames.Add("ElementalDemo", DesktopPlats);
+        NonCodeProjectNames.Add("InfiltratorDemo", DesktopPlats);
+        NonCodeProjectNames.Add("HoverShip", DesktopPlats);
+        NonCodeProjectNames.Add("BlueprintOffice", DesktopPlats);
+        NonCodeProjectNames.Add("ReflectionsSubway", DesktopPlats);
+        NonCodeProjectNames.Add("ElementalVR", DesktopPlats);
+        NonCodeProjectNames.Add("CouchKnights", DesktopPlats);
+        NonCodeProjectNames.Add("Stylized", DesktopPlats);
         //NonCodeProjectNames.Add("Landscape", new List<UnrealTargetPlatform> { HostPlatform });
-        NonCodeProjectNames.Add("MatineeFightScene", new List<UnrealTargetPlatform> { HostPlatform });
-        NonCodeProjectNames.Add("RealisticRendering", new List<UnrealTargetPlatform> { HostPlatform });
-        NonCodeProjectNames.Add("EffectsCave", new List<UnrealTargetPlatform> { HostPlatform });
-        NonCodeProjectNames.Add("GDC2014", new List<UnrealTargetPlatform> { HostPlatform });
-        NonCodeProjectNames.Add("ContentExamples", new List<UnrealTargetPlatform> { HostPlatform });
-        NonCodeProjectNames.Add("PhysicsPirateShip", new List<UnrealTargetPlatform> { HostPlatform });
-        NonCodeProjectNames.Add("TowerDefenseGame", new List<UnrealTargetPlatform> { HostPlatform });
-        NonCodeProjectNames.Add("LandscapeMountains", new List<UnrealTargetPlatform> { HostPlatform });
-        NonCodeProjectNames.Add("MorphTargets", new List<UnrealTargetPlatform> { HostPlatform });
-        NonCodeProjectNames.Add("PostProcessMatinee", new List<UnrealTargetPlatform> { HostPlatform });
-        NonCodeProjectNames.Add("SciFiHallway", new List<UnrealTargetPlatform> { HostPlatform });
+        NonCodeProjectNames.Add("MatineeFightScene", DesktopPlats);
+        NonCodeProjectNames.Add("RealisticRendering", DesktopPlats);
+        NonCodeProjectNames.Add("EffectsCave", DesktopPlats);
+        NonCodeProjectNames.Add("GDC2014", DesktopPlats);
+        NonCodeProjectNames.Add("ContentExamples", DesktopPlats);
+        NonCodeProjectNames.Add("PhysicsPirateShip", DesktopPlats);
+        NonCodeProjectNames.Add("TowerDefenseGame", DesktopPlats);
+        NonCodeProjectNames.Add("LandscapeMountains", DesktopPlats);
+        NonCodeProjectNames.Add("MorphTargets", DesktopPlats);
+        NonCodeProjectNames.Add("PostProcessMatinee", DesktopPlats);
+        NonCodeProjectNames.Add("SciFiHallway", DesktopPlats);
 
         List<UnrealTargetPlatform> MobilePlats = null;
         if (HostPlatform == UnrealTargetPlatform.Mac)
diff --git a/Engine/Source/UE4Game.Target.cs b/Engine/Source/UE4Game.Target.cs
--- a/Engine/Source/UE4Game.Target.cs
+++ b/Engine/Source/UE4Game.Target.cs
@@ -90,7 +90,7 @@ public class UE4GameTarget : TargetRules
         {
             return new List<UnrealTargetPlatform> { HostPlatform, UnrealTargetPlatform.IOS };
         }
-        return new List<UnrealTargetPlatform> { HostPlatform, UnrealTargetPlatform.Win32, UnrealTargetPlatform.IOS, UnrealTargetPlatform.XboxOne, UnrealTargetPlatform.PS4, UnrealTargetPlatform.Android };
+        return new List<UnrealTargetPlatform> { HostPlatform, UnrealTargetPlatform.Win32, UnrealTargetPlatform.IOS, UnrealTargetPlatform.XboxOne, UnrealTargetPlatform.PS4, UnrealTargetPlatform.Android, UnrealTargetPlatform.Linux };
     }
     public override List<UnrealTargetConfiguration> GUBP_GetConfigs_MonolithicOnly(UnrealTargetPlatform HostPlatform, UnrealTargetPlatform Platform)
     {
