build script, add animationsamplepack

diff --git a/Engine/Source/UE4Editor.Target.cs b/Engine/Source/UE4Editor.Target.cs
--- a/Engine/Source/UE4Editor.Target.cs
+++ b/Engine/Source/UE4Editor.Target.cs
@@ -109,6 +109,7 @@ public class UE4EditorTarget : TargetRules
         NonCodeProjectNames.Add("TappyChicken", MobilePlats);
         NonCodeProjectNames.Add("SwingNinja", MobilePlats);
         NonCodeProjectNames.Add("MobileTemple", MobilePlats);
+        NonCodeProjectNames.Add("AnimationStarterPack", MobilePlats);
 
         NonCodeProjectNames.Add("StarterContent", MobilePlats);
         NonCodeProjectNames.Add("TP_FirstPersonBP", MobilePlats);
