build script, remove infiltrator

diff --git a/Engine/Source/UE4Editor.Target.cs b/Engine/Source/UE4Editor.Target.cs
--- a/Engine/Source/UE4Editor.Target.cs
+++ b/Engine/Source/UE4Editor.Target.cs
@@ -74,7 +74,7 @@ public class UE4EditorTarget : TargetRules
         }
 
         NonCodeProjectNames.Add("ElementalDemo", DesktopPlats);
-        NonCodeProjectNames.Add("InfiltratorDemo", DesktopPlats);
+        //NonCodeProjectNames.Add("InfiltratorDemo", DesktopPlats);
         NonCodeProjectNames.Add("HoverShip", DesktopPlats);
         NonCodeProjectNames.Add("BlueprintOffice", DesktopPlats);
         NonCodeProjectNames.Add("ReflectionsSubway", DesktopPlats);
