remove some of the Tappy Chicken changes
#ue4
#html5

diff --git a/Engine/Source/Programs/UnrealBuildTool/HTML5/HTML5ToolChain.cs b/Engine/Source/Programs/UnrealBuildTool/HTML5/HTML5ToolChain.cs
--- a/Engine/Source/Programs/UnrealBuildTool/HTML5/HTML5ToolChain.cs
+++ b/Engine/Source/Programs/UnrealBuildTool/HTML5/HTML5ToolChain.cs
@@ -75,7 +75,7 @@ namespace UnrealBuildTool
 
             // we have to specify the full amount of memory with Asm.JS (1.5 G)
             // I wonder if there's a per game way to change this. 
-			int TotalMemory = 256 * 1024 * 1024;
+			int TotalMemory = 512 * 1024 * 1024;
             Result += " -s TOTAL_MEMORY=" + TotalMemory.ToString();
 
             // no need for exceptions
diff --git a/Engine/Source/Runtime/Engine/Private/Vehicles/WheeledVehicleMovementComponent.cpp b/Engine/Source/Runtime/Engine/Private/Vehicles/WheeledVehicleMovementComponent.cpp
--- a/Engine/Source/Runtime/Engine/Private/Vehicles/WheeledVehicleMovementComponent.cpp
+++ b/Engine/Source/Runtime/Engine/Private/Vehicles/WheeledVehicleMovementComponent.cpp
@@ -103,7 +103,7 @@ UWheeledVehicleMovementComponent::UWheeledVehicleMovementComponent(const class F
 	SteeringInputRate.FallRate = 5.0f;
 
 
-#if WITH_PHYSX
+#if WITH_VEHICLE
 	// tire load filtering
 	PxVehicleTireLoadFilterData PTireLoadFilterDef;
 	MinNormalizedTireLoad = PTireLoadFilterDef.mMinNormalisedLoad;
diff --git a/Engine/Source/UE4Game.Target.cs b/Engine/Source/UE4Game.Target.cs
--- a/Engine/Source/UE4Game.Target.cs
+++ b/Engine/Source/UE4Game.Target.cs
@@ -96,13 +96,14 @@ public class UE4GameTarget : TargetRules
 		if (Target.Platform == UnrealTargetPlatform.HTML5)
 		{
 			// to make Tappy Chicken as small as possible we excluded some items from the engine.
-			UEBuildConfiguration.bCompileRecast = false;
+			// uncomment below to make a smaller HTML5 build
+/*			UEBuildConfiguration.bCompileRecast = false;
 			UEBuildConfiguration.bCompileSpeedTree = false;
 			UEBuildConfiguration.bCompileAPEX = false;
 			UEBuildConfiguration.bCompileLeanAndMeanUE = true;
 			UEBuildConfiguration.bCompilePhysXVehicle = false;
 			UEBuildConfiguration.bCompileForSize = true;
-			UEBuildConfiguration.bCompileFreeType = false;
+			UEBuildConfiguration.bCompileFreeType = false;*/
 		}
 	}
 	public override List<UnrealTargetPlatform> GUBP_GetPlatforms_MonolithicOnly(UnrealTargetPlatform HostPlatform)
