build script, turn on device testing

diff --git a/Engine/Source/UE4Editor.Target.cs b/Engine/Source/UE4Editor.Target.cs
--- a/Engine/Source/UE4Editor.Target.cs
+++ b/Engine/Source/UE4Editor.Target.cs
@@ -128,10 +128,10 @@ public class UE4EditorTarget : TargetRules
         NonCodeProjectNames.Add("TappyChicken",
             new List<GUBPFormalBuild>
             {
+                    new GUBPFormalBuild(UnrealTargetPlatform.Android, UnrealTargetConfiguration.Shipping, true),
+                    new GUBPFormalBuild(UnrealTargetPlatform.Android, UnrealTargetConfiguration.Test, true),                    
                     new GUBPFormalBuild(UnrealTargetPlatform.IOS, UnrealTargetConfiguration.Shipping),
-                    new GUBPFormalBuild(UnrealTargetPlatform.IOS, UnrealTargetConfiguration.Test),
-                    new GUBPFormalBuild(UnrealTargetPlatform.Android, UnrealTargetConfiguration.Shipping /*, true*/),
-                    new GUBPFormalBuild(UnrealTargetPlatform.Android, UnrealTargetConfiguration.Test)
+                    new GUBPFormalBuild(UnrealTargetPlatform.IOS, UnrealTargetConfiguration.Test),                    
             }
         );
         return NonCodeProjectNames;
