buildscript, adding html5 builds

diff --git a/Engine/Source/UE4Game.Target.cs b/Engine/Source/UE4Game.Target.cs
--- a/Engine/Source/UE4Game.Target.cs
+++ b/Engine/Source/UE4Game.Target.cs
@@ -112,7 +112,7 @@ public class UE4GameTarget : TargetRules
         {
             return new List<UnrealTargetPlatform> { HostPlatform, UnrealTargetPlatform.IOS };
         }
-        return new List<UnrealTargetPlatform> { HostPlatform, UnrealTargetPlatform.Win32, UnrealTargetPlatform.IOS, UnrealTargetPlatform.XboxOne, UnrealTargetPlatform.PS4, UnrealTargetPlatform.Android, UnrealTargetPlatform.Linux };
+        return new List<UnrealTargetPlatform> { HostPlatform, UnrealTargetPlatform.Win32, UnrealTargetPlatform.IOS, UnrealTargetPlatform.XboxOne, UnrealTargetPlatform.PS4, UnrealTargetPlatform.Android, UnrealTargetPlatform.Linux, UnrealTargetPlatform.HTML5 };
     }
     public override List<UnrealTargetConfiguration> GUBP_GetConfigs_MonolithicOnly(UnrealTargetPlatform HostPlatform, UnrealTargetPlatform Platform)
     {
