buildscript, turn on ios device testing

diff --git a/Engine/Source/UE4Editor.Target.cs b/Engine/Source/UE4Editor.Target.cs
--- a/Engine/Source/UE4Editor.Target.cs
+++ b/Engine/Source/UE4Editor.Target.cs
@@ -130,8 +130,8 @@ public class UE4EditorTarget : TargetRules
             {
                     new GUBPFormalBuild(UnrealTargetPlatform.Android, UnrealTargetConfiguration.Shipping, true),
                     new GUBPFormalBuild(UnrealTargetPlatform.Android, UnrealTargetConfiguration.Test, true),                    
-                    new GUBPFormalBuild(UnrealTargetPlatform.IOS, UnrealTargetConfiguration.Shipping),
-                    new GUBPFormalBuild(UnrealTargetPlatform.IOS, UnrealTargetConfiguration.Test),                    
+                    new GUBPFormalBuild(UnrealTargetPlatform.IOS, UnrealTargetConfiguration.Shipping, true),
+                    new GUBPFormalBuild(UnrealTargetPlatform.IOS, UnrealTargetConfiguration.Test, true),
             }
         );
         return NonCodeProjectNames;
