Add 2D sidescroller into the Rocket templates list.

diff --git a/Engine/Source/UE4Editor.Target.cs b/Engine/Source/UE4Editor.Target.cs
--- a/Engine/Source/UE4Editor.Target.cs
+++ b/Engine/Source/UE4Editor.Target.cs
@@ -112,6 +112,7 @@ public class UE4EditorTarget : TargetRules
         NonCodeProjectNames.Add("AnimationStarterPack", MobilePlats);
 
         NonCodeProjectNames.Add("StarterContent", MobilePlats);
+		NonCodeProjectNames.Add("TP_2DSideScrollerBP", MobilePlats);
         NonCodeProjectNames.Add("TP_FirstPersonBP", MobilePlats);
         NonCodeProjectNames.Add("TP_FlyingBP", MobilePlats);
         NonCodeProjectNames.Add("TP_RollingBP", MobilePlats);
