Enable iOS on PC as part of RocketBuild. Also fix name of AnimStarterPack in the non-code project list.

diff --git a/Engine/Source/Programs/UnrealBuildTool/IOS/UEBuildIOS.cs b/Engine/Source/Programs/UnrealBuildTool/IOS/UEBuildIOS.cs
--- a/Engine/Source/Programs/UnrealBuildTool/IOS/UEBuildIOS.cs
+++ b/Engine/Source/Programs/UnrealBuildTool/IOS/UEBuildIOS.cs
@@ -31,10 +31,6 @@ namespace UnrealBuildTool
 		 */
 		public override void RegisterBuildPlatform()
 		{
-            if ((UnrealBuildTool.RunningRocket() || UnrealBuildTool.BuildingRocket()) && ExternalExecution.GetRuntimePlatform() != UnrealTargetPlatform.Mac)
-            {
-                return;
-            }
 			// Register this build platform for IOS
 			Log.TraceVerbose("        Registering for {0}", UnrealTargetPlatform.IOS.ToString());
 			UEBuildPlatform.RegisterBuildPlatform(UnrealTargetPlatform.IOS, this);
diff --git a/Engine/Source/UE4Editor.Target.cs b/Engine/Source/UE4Editor.Target.cs
--- a/Engine/Source/UE4Editor.Target.cs
+++ b/Engine/Source/UE4Editor.Target.cs
@@ -101,7 +101,7 @@ public class UE4EditorTarget : TargetRules
         }
         else
         {
-            MobilePlats = new List<UnrealTargetPlatform> { HostPlatform, UnrealTargetPlatform.Android };
+            MobilePlats = new List<UnrealTargetPlatform> { HostPlatform, UnrealTargetPlatform.Android, UnrealTargetPlatform.IOS };
         }
 
         NonCodeProjectNames.Add("BlackJack", MobilePlats);
@@ -109,7 +109,7 @@ public class UE4EditorTarget : TargetRules
         NonCodeProjectNames.Add("TappyChicken", MobilePlats);
         NonCodeProjectNames.Add("SwingNinja", MobilePlats);
         NonCodeProjectNames.Add("MobileTemple", MobilePlats);
-        NonCodeProjectNames.Add("AnimationStarterPack", MobilePlats);
+        NonCodeProjectNames.Add("AnimStarterPack", MobilePlats);
 
         NonCodeProjectNames.Add("StarterContent", MobilePlats);
 		NonCodeProjectNames.Add("TP_2DSideScrollerBP", MobilePlats);
