disable HTML5 builds for the time being until I can update the build machines to 1.21
#ue4
#html5

diff --git a/Engine/Source/UE4Game.Target.cs b/Engine/Source/UE4Game.Target.cs
--- a/Engine/Source/UE4Game.Target.cs
+++ b/Engine/Source/UE4Game.Target.cs
@@ -102,8 +102,8 @@ public class UE4GameTarget : TargetRules
 			UEBuildConfiguration.bCompileAPEX = false;
 			UEBuildConfiguration.bCompileLeanAndMeanUE = true;
 			UEBuildConfiguration.bCompilePhysXVehicle = false;
-			UEBuildConfiguration.bCompileFreeType = false;*/
-			UEBuildConfiguration.bCompileForSize = true;
+			UEBuildConfiguration.bCompileFreeType = false;
+			UEBuildConfiguration.bCompileForSize = true;*/
 		}
 	}
 	public override List<UnrealTargetPlatform> GUBP_GetPlatforms_MonolithicOnly(UnrealTargetPlatform HostPlatform)
@@ -112,7 +112,7 @@ public class UE4GameTarget : TargetRules
         {
             return new List<UnrealTargetPlatform> { HostPlatform, UnrealTargetPlatform.IOS };
         }
-        return new List<UnrealTargetPlatform> { HostPlatform, UnrealTargetPlatform.Win32, UnrealTargetPlatform.IOS, UnrealTargetPlatform.XboxOne, UnrealTargetPlatform.PS4, UnrealTargetPlatform.Android, UnrealTargetPlatform.Linux, UnrealTargetPlatform.HTML5 };
+        return new List<UnrealTargetPlatform> { HostPlatform, UnrealTargetPlatform.Win32, UnrealTargetPlatform.IOS, UnrealTargetPlatform.XboxOne, UnrealTargetPlatform.PS4, UnrealTargetPlatform.Android, UnrealTargetPlatform.Linux };
     }
     public override List<UnrealTargetConfiguration> GUBP_GetConfigs_MonolithicOnly(UnrealTargetPlatform HostPlatform, UnrealTargetPlatform Platform)
     {
