Fix PlatformerGame compile errors in Rocket build: GameMenuBuilder is not built as part of the engine because nothing depends on it, so it's not included in Rocket. Add it into the extra modules list for editor and game targets.

diff --git a/Engine/Source/UE4Editor.Target.cs b/Engine/Source/UE4Editor.Target.cs
--- a/Engine/Source/UE4Editor.Target.cs
+++ b/Engine/Source/UE4Editor.Target.cs
@@ -21,6 +21,7 @@ public class UE4EditorTarget : TargetRules
 		ref List<string> OutExtraModuleNames
 		)
 	{
+		OutExtraModuleNames.Add("GameMenuBuilder");
 		if ((Target.Platform == UnrealTargetPlatform.Win32) || (Target.Platform == UnrealTargetPlatform.Win64))
 		{
 			OutExtraModuleNames.Add("OnlineSubsystemNull");
diff --git a/Engine/Source/UE4Game.Target.cs b/Engine/Source/UE4Game.Target.cs
--- a/Engine/Source/UE4Game.Target.cs
+++ b/Engine/Source/UE4Game.Target.cs
@@ -33,6 +33,7 @@ public class UE4GameTarget : TargetRules
 
 		if (UnrealBuildTool.UnrealBuildTool.BuildingRocket())
 		{
+			OutExtraModuleNames.Add("GameMenuBuilder");
 			if ((Target.Platform == UnrealTargetPlatform.Win32) || (Target.Platform == UnrealTargetPlatform.Win64))
 			{
 				OutExtraModuleNames.Add("OnlineSubsystemNull");
