Added BlackJack, MemoryGame, SwingingNinja, and MobileTemple to gubp build scripts.

diff --git a/Engine/Source/UE4Editor.Target.cs b/Engine/Source/UE4Editor.Target.cs
--- a/Engine/Source/UE4Editor.Target.cs
+++ b/Engine/Source/UE4Editor.Target.cs
@@ -126,16 +126,50 @@ public class UE4EditorTarget : TargetRules
     }
     public override Dictionary<string, List<GUBPFormalBuild>> GUBP_GetNonCodeFormalBuilds_BaseEditorTypeOnly()
     {
-        var NonCodeProjectNames = new Dictionary<string, List<GUBPFormalBuild>>();
-        NonCodeProjectNames.Add("TappyChicken",
-            new List<GUBPFormalBuild>
+        var TestShippingFormalBuildSettingsTest = new List<GUBPFormalBuild>
             {
                     new GUBPFormalBuild(UnrealTargetPlatform.Android, UnrealTargetConfiguration.Shipping, true),
                     new GUBPFormalBuild(UnrealTargetPlatform.Android, UnrealTargetConfiguration.Test, true),                    
                     new GUBPFormalBuild(UnrealTargetPlatform.IOS, UnrealTargetConfiguration.Shipping, true),
                     new GUBPFormalBuild(UnrealTargetPlatform.IOS, UnrealTargetConfiguration.Test, true),
-            }
-        );
+            };
+
+        var NonCodeProjectNames = new Dictionary<string, List<GUBPFormalBuild>>();
+        NonCodeProjectNames.Add("TappyChicken", TestShippingFormalBuildSettingsTest);
+
+
+        var TestShippingFormalBuildSettings = new List<GUBPFormalBuild>
+            {
+                    new GUBPFormalBuild(UnrealTargetPlatform.Android, UnrealTargetConfiguration.Test, false),
+                    new GUBPFormalBuild(UnrealTargetPlatform.IOS, UnrealTargetConfiguration.Test, false),
+            };
+
+
+        NonCodeProjectNames.Add("StrategyGame", TestShippingFormalBuildSettings );
+
+
+        var TestFormalBuildSettings = new List<GUBPFormalBuild>
+            {
+                    new GUBPFormalBuild(UnrealTargetPlatform.Android, UnrealTargetConfiguration.Test, false),
+                    new GUBPFormalBuild(UnrealTargetPlatform.IOS, UnrealTargetConfiguration.Test, false),
+            };
+
+        NonCodeProjectNames.Add("BlackJack", TestFormalBuildSettings);
+        NonCodeProjectNames.Add("MemoryGame", TestFormalBuildSettings);
+        NonCodeProjectNames.Add("SwingNinja", TestFormalBuildSettings);
+        NonCodeProjectNames.Add("MobileTemple", TestFormalBuildSettings);
+
+
+        /*NonCodeProjectNames.Add("AnimStarterPack",FormalBuildSettings);
+        NonCodeProjectNames.Add("StarterContent",FormalBuildSettings);
+        NonCodeProjectNames.Add("TP_2DSideScrollerBP",FormalBuildSettings);
+        NonCodeProjectNames.Add("TP_FirstPersonBP",FormalBuildSettings);
+        NonCodeProjectNames.Add("TP_FlyingBP",FormalBuildSettings);
+        NonCodeProjectNames.Add("TP_RollingBP",FormalBuildSettings);
+        NonCodeProjectNames.Add("TP_SideScrollerBP",FormalBuildSettings);
+        NonCodeProjectNames.Add("TP_TopDownBP",FormalBuildSettings);
+        NonCodeProjectNames.Add("TP_VehicleBP",FormalBuildSettings);*/
+
         return NonCodeProjectNames;
     }
 }
