#UE4 #HTML5

    TTP 342328:  HTML5: Packaging Blueprint projects fails
    re - enablng GUBP for HTML5.

diff --git a/Engine/Source/UE4Game.Target.cs b/Engine/Source/UE4Game.Target.cs
--- a/Engine/Source/UE4Game.Target.cs
+++ b/Engine/Source/UE4Game.Target.cs
@@ -29,7 +29,7 @@ public class UE4GameTarget : TargetRules
 		{
 			OutExtraModuleNames.Add("OnlineSubsystemNull");
 		}
-		
+
 
 		if (UnrealBuildTool.UnrealBuildTool.BuildingRocket())
 		{
@@ -90,7 +90,7 @@ public class UE4GameTarget : TargetRules
 			OutLinkEnvironmentConfiguration.bHasExports = false;
 		}
 		else
-		{ 
+		{
 			// Tag it as a UE4Game build
 			OutCPPEnvironmentConfiguration.Definitions.Add("UE4GAME=1");
 		}
@@ -113,7 +113,7 @@ public class UE4GameTarget : TargetRules
         {
             return new List<UnrealTargetPlatform> { HostPlatform, UnrealTargetPlatform.IOS };
         }
-        return new List<UnrealTargetPlatform> { HostPlatform, UnrealTargetPlatform.Win32, UnrealTargetPlatform.IOS, UnrealTargetPlatform.XboxOne, UnrealTargetPlatform.PS4, UnrealTargetPlatform.Android, UnrealTargetPlatform.Linux };
+        return new List<UnrealTargetPlatform> { HostPlatform, UnrealTargetPlatform.Win32, UnrealTargetPlatform.IOS, UnrealTargetPlatform.XboxOne, UnrealTargetPlatform.PS4, UnrealTargetPlatform.Android, UnrealTargetPlatform.Linux, UnrealTargetPlatform.HTML5 };
     }
     public override List<UnrealTargetConfiguration> GUBP_GetConfigs_MonolithicOnly(UnrealTargetPlatform HostPlatform, UnrealTargetPlatform Platform)
     {
