Disabled physicspirateship sample automated build.

diff --git a/Engine/Source/UE4Editor.Target.cs b/Engine/Source/UE4Editor.Target.cs
--- a/Engine/Source/UE4Editor.Target.cs
+++ b/Engine/Source/UE4Editor.Target.cs
@@ -88,7 +88,7 @@ public class UE4EditorTarget : TargetRules
         NonCodeProjectNames.Add("EffectsCave", DesktopPlats);
         NonCodeProjectNames.Add("GDC2014", DesktopPlats);
         NonCodeProjectNames.Add("ContentExamples", DesktopPlats);
-        NonCodeProjectNames.Add("PhysicsPirateShip", DesktopPlats);
+        // NonCodeProjectNames.Add("PhysicsPirateShip", DesktopPlats);
         NonCodeProjectNames.Add("TowerDefenseGame", DesktopPlats);
         NonCodeProjectNames.Add("LandscapeMountains", DesktopPlats);
         NonCodeProjectNames.Add("MorphTargets", DesktopPlats);
