Added strategy game to the auto generated builds list

diff --git a/Engine/Source/UE4Editor.Target.cs b/Engine/Source/UE4Editor.Target.cs
--- a/Engine/Source/UE4Editor.Target.cs
+++ b/Engine/Source/UE4Editor.Target.cs
@@ -111,6 +111,7 @@ public class UE4EditorTarget : TargetRules
         NonCodeProjectNames.Add("SwingNinja", MobilePlats);
         NonCodeProjectNames.Add("MobileTemple", MobilePlats);
         NonCodeProjectNames.Add("AnimStarterPack", MobilePlats);
+        NonCodeProjectNames.Add("StrategyGame", MobilePlats);
 
         NonCodeProjectNames.Add("StarterContent", MobilePlats);
 		NonCodeProjectNames.Add("TP_2DSideScrollerBP", MobilePlats);
@@ -130,8 +131,10 @@ public class UE4EditorTarget : TargetRules
             {
                     new GUBPFormalBuild(UnrealTargetPlatform.Android, UnrealTargetConfiguration.Shipping, true),
                     new GUBPFormalBuild(UnrealTargetPlatform.Android, UnrealTargetConfiguration.Test, true),                    
+                    new GUBPFormalBuild(UnrealTargetPlatform.Android, UnrealTargetConfiguration.Development, true),          
                     new GUBPFormalBuild(UnrealTargetPlatform.IOS, UnrealTargetConfiguration.Shipping, true),
                     new GUBPFormalBuild(UnrealTargetPlatform.IOS, UnrealTargetConfiguration.Test, true),
+                    new GUBPFormalBuild(UnrealTargetPlatform.IOS, UnrealTargetConfiguration.Development, true),          
             };
 
         var NonCodeProjectNames = new Dictionary<string, List<GUBPFormalBuild>>();
@@ -140,8 +143,12 @@ public class UE4EditorTarget : TargetRules
 
         var TestShippingFormalBuildSettings = new List<GUBPFormalBuild>
             {
-                    new GUBPFormalBuild(UnrealTargetPlatform.Android, UnrealTargetConfiguration.Test, false),
+                    new GUBPFormalBuild(UnrealTargetPlatform.Android, UnrealTargetConfiguration.Shipping, false),
+                    new GUBPFormalBuild(UnrealTargetPlatform.Android, UnrealTargetConfiguration.Test, false),                    
+                    new GUBPFormalBuild(UnrealTargetPlatform.Android, UnrealTargetConfiguration.Development, false),          
+                    new GUBPFormalBuild(UnrealTargetPlatform.IOS, UnrealTargetConfiguration.Shipping, false),
                     new GUBPFormalBuild(UnrealTargetPlatform.IOS, UnrealTargetConfiguration.Test, false),
+                    new GUBPFormalBuild(UnrealTargetPlatform.IOS, UnrealTargetConfiguration.Development, false),  
             };
 
 
@@ -159,17 +166,6 @@ public class UE4EditorTarget : TargetRules
         NonCodeProjectNames.Add("SwingNinja", TestFormalBuildSettings);
         NonCodeProjectNames.Add("MobileTemple", TestFormalBuildSettings);
 
-
-        /*NonCodeProjectNames.Add("AnimStarterPack",FormalBuildSettings);
-        NonCodeProjectNames.Add("StarterContent",FormalBuildSettings);
-        NonCodeProjectNames.Add("TP_2DSideScrollerBP",FormalBuildSettings);
-        NonCodeProjectNames.Add("TP_FirstPersonBP",FormalBuildSettings);
-        NonCodeProjectNames.Add("TP_FlyingBP",FormalBuildSettings);
-        NonCodeProjectNames.Add("TP_RollingBP",FormalBuildSettings);
-        NonCodeProjectNames.Add("TP_SideScrollerBP",FormalBuildSettings);
-        NonCodeProjectNames.Add("TP_TopDownBP",FormalBuildSettings);
-        NonCodeProjectNames.Add("TP_VehicleBP",FormalBuildSettings);*/
-
         return NonCodeProjectNames;
     }
 }
