Added some more builds to be generated.

diff --git a/Engine/Source/UE4Editor.Target.cs b/Engine/Source/UE4Editor.Target.cs
--- a/Engine/Source/UE4Editor.Target.cs
+++ b/Engine/Source/UE4Editor.Target.cs
@@ -111,7 +111,6 @@ public class UE4EditorTarget : TargetRules
         NonCodeProjectNames.Add("SwingNinja", MobilePlats);
         NonCodeProjectNames.Add("MobileTemple", MobilePlats);
         NonCodeProjectNames.Add("AnimStarterPack", MobilePlats);
-        NonCodeProjectNames.Add("StrategyGame", MobilePlats);
 
         NonCodeProjectNames.Add("StarterContent", MobilePlats);
 		NonCodeProjectNames.Add("TP_2DSideScrollerBP", MobilePlats);
@@ -127,7 +126,7 @@ public class UE4EditorTarget : TargetRules
     }
     public override Dictionary<string, List<GUBPFormalBuild>> GUBP_GetNonCodeFormalBuilds_BaseEditorTypeOnly()
     {
-        var TestShippingFormalBuildSettingsTest = new List<GUBPFormalBuild>
+        var TappyChickenBuildSettings = new List<GUBPFormalBuild>
             {
                     new GUBPFormalBuild(UnrealTargetPlatform.Android, UnrealTargetConfiguration.Shipping, true),
                     new GUBPFormalBuild(UnrealTargetPlatform.Android, UnrealTargetConfiguration.Test, true),                    
@@ -138,33 +137,59 @@ public class UE4EditorTarget : TargetRules
             };
 
         var NonCodeProjectNames = new Dictionary<string, List<GUBPFormalBuild>>();
-        NonCodeProjectNames.Add("TappyChicken", TestShippingFormalBuildSettingsTest);
+        NonCodeProjectNames.Add("TappyChicken", TappyChickenBuildSettings);
 
 
-        var TestShippingFormalBuildSettings = new List<GUBPFormalBuild>
+
+        var MobileBuildSettings = new List<GUBPFormalBuild>
             {
-                    new GUBPFormalBuild(UnrealTargetPlatform.Android, UnrealTargetConfiguration.Shipping, false),
-                    new GUBPFormalBuild(UnrealTargetPlatform.Android, UnrealTargetConfiguration.Test, false),                    
-                    new GUBPFormalBuild(UnrealTargetPlatform.Android, UnrealTargetConfiguration.Development, false),          
-                    new GUBPFormalBuild(UnrealTargetPlatform.IOS, UnrealTargetConfiguration.Shipping, false),
+                    new GUBPFormalBuild(UnrealTargetPlatform.Android, UnrealTargetConfiguration.Test, false),
                     new GUBPFormalBuild(UnrealTargetPlatform.IOS, UnrealTargetConfiguration.Test, false),
-                    new GUBPFormalBuild(UnrealTargetPlatform.IOS, UnrealTargetConfiguration.Development, false),  
             };
 
+        NonCodeProjectNames.Add("BlackJack", MobileBuildSettings);
+        NonCodeProjectNames.Add("MemoryGame", MobileBuildSettings);
+        NonCodeProjectNames.Add("SwingNinja", MobileBuildSettings);
+        NonCodeProjectNames.Add("MobileTemple", MobileBuildSettings);
+
+        var PCBuildSettings = new List<GUBPFormalBuild>
+            {
+                    new GUBPFormalBuild(UnrealTargetPlatform.Win32, UnrealTargetConfiguration.Test, false),
+                    new GUBPFormalBuild(UnrealTargetPlatform.Win64, UnrealTargetConfiguration.Test, false),
+                    new GUBPFormalBuild(UnrealTargetPlatform.Mac, UnrealTargetConfiguration.Test, false),
+            };
 
-        NonCodeProjectNames.Add("StrategyGame", TestShippingFormalBuildSettings );
+        NonCodeProjectNames.Add("BlueprintOffice", PCBuildSettings);
+        NonCodeProjectNames.Add("ContentExamples", PCBuildSettings);
+        NonCodeProjectNames.Add("CouchKnights", PCBuildSettings);
+        NonCodeProjectNames.Add("EffectsCave", PCBuildSettings);
+        NonCodeProjectNames.Add("LandscapeMountains", PCBuildSettings);
+        NonCodeProjectNames.Add("GDC2014", PCBuildSettings);
+        NonCodeProjectNames.Add("MatineeFightScene", PCBuildSettings);
+        NonCodeProjectNames.Add("RealisticRendering", PCBuildSettings);
+        NonCodeProjectNames.Add("ReflectionsSubway", PCBuildSettings);
+        NonCodeProjectNames.Add("Stylized", PCBuildSettings);
 
 
-        var TestFormalBuildSettings = new List<GUBPFormalBuild>
+        var ElementalBuildSettings = new List<GUBPFormalBuild>
             {
-                    new GUBPFormalBuild(UnrealTargetPlatform.Android, UnrealTargetConfiguration.Test, false),
-                    new GUBPFormalBuild(UnrealTargetPlatform.IOS, UnrealTargetConfiguration.Test, false),
+                    new GUBPFormalBuild(UnrealTargetPlatform.Win32, UnrealTargetConfiguration.Test, false),
+                    new GUBPFormalBuild(UnrealTargetPlatform.Win64, UnrealTargetConfiguration.Test, false),
+                    new GUBPFormalBuild(UnrealTargetPlatform.Mac, UnrealTargetConfiguration.Test, false),
+                    new GUBPFormalBuild(UnrealTargetPlatform.XboxOne, UnrealTargetConfiguration.Test, false),
+                    new GUBPFormalBuild(UnrealTargetPlatform.PS4, UnrealTargetConfiguration.Test, false),
+                    
+                    new GUBPFormalBuild(UnrealTargetPlatform.XboxOne, UnrealTargetConfiguration.Shipping, false),
+                    new GUBPFormalBuild(UnrealTargetPlatform.PS4, UnrealTargetConfiguration.Shipping, false),
+                    
+                    new GUBPFormalBuild(UnrealTargetPlatform.XboxOne, UnrealTargetConfiguration.Development, false),
+                    new GUBPFormalBuild(UnrealTargetPlatform.PS4, UnrealTargetConfiguration.Development, false),
             };
 
-        NonCodeProjectNames.Add("BlackJack", TestFormalBuildSettings);
-        NonCodeProjectNames.Add("MemoryGame", TestFormalBuildSettings);
-        NonCodeProjectNames.Add("SwingNinja", TestFormalBuildSettings);
-        NonCodeProjectNames.Add("MobileTemple", TestFormalBuildSettings);
+
+        NonCodeProjectNames.Add("ElementalDemo", ElementalBuildSettings);
+
+
 
         return NonCodeProjectNames;
     }
