build script, remove asset pack from builds

diff --git a/Engine/Source/UE4Editor.Target.cs b/Engine/Source/UE4Editor.Target.cs
--- a/Engine/Source/UE4Editor.Target.cs
+++ b/Engine/Source/UE4Editor.Target.cs
@@ -110,7 +110,6 @@ public class UE4EditorTarget : TargetRules
 		NonCodeProjectNames.Add("TappyChicken", MobilePlats);
         NonCodeProjectNames.Add("SwingNinja", MobilePlats);
         NonCodeProjectNames.Add("MobileTemple", MobilePlats);
-        NonCodeProjectNames.Add("AnimStarterPack", MobilePlats);
 
         NonCodeProjectNames.Add("StarterContent", MobilePlats);
 		NonCodeProjectNames.Add("TP_2DSideScrollerBP", MobilePlats);
