More uniform handling of GUBP platforms.

Specifies platforms more explicitly, making things easier when an "unexpected" host platform is passed.

diff --git a/Engine/Source/UE4Client.Target.cs b/Engine/Source/UE4Client.Target.cs
--- a/Engine/Source/UE4Client.Target.cs
+++ b/Engine/Source/UE4Client.Target.cs
@@ -53,11 +53,28 @@ public class UE4ClientTarget : TargetRules
 
     public override List<UnrealTargetPlatform> GUBP_GetPlatforms_MonolithicOnly(UnrealTargetPlatform HostPlatform)
     {
-        if (HostPlatform == UnrealTargetPlatform.Mac)
-        {
-            return new List<UnrealTargetPlatform> { HostPlatform, UnrealTargetPlatform.IOS };
-        }
-        return new List<UnrealTargetPlatform> { HostPlatform, UnrealTargetPlatform.Win32, UnrealTargetPlatform.Linux /*, UnrealTargetPlatform.IOS, UnrealTargetPlatform.XboxOne, UnrealTargetPlatform.PS4 */};
+		List<UnrealTargetPlatform> Platforms = null;
+
+		switch(HostPlatform)
+		{
+			case UnrealTargetPlatform.Mac:
+				Platforms = new List<UnrealTargetPlatform> { HostPlatform, UnrealTargetPlatform.IOS };
+				break;
+
+			case UnrealTargetPlatform.Linux:
+				Platforms = new List<UnrealTargetPlatform> { HostPlatform };
+				break;
+
+			case UnrealTargetPlatform.Win64:
+				Platforms = new List<UnrealTargetPlatform> { HostPlatform, UnrealTargetPlatform.Win32, UnrealTargetPlatform.Linux /*, UnrealTargetPlatform.IOS, UnrealTargetPlatform.XboxOne, UnrealTargetPlatform.PS4 */};
+				break;
+
+			default:
+				Platforms = new List<UnrealTargetPlatform>();
+				break;
+		}
+
+		return Platforms;
     }
 
     public override List<UnrealTargetConfiguration> GUBP_GetConfigs_MonolithicOnly(UnrealTargetPlatform HostPlatform, UnrealTargetPlatform Platform)
diff --git a/Engine/Source/UE4Game.Target.cs b/Engine/Source/UE4Game.Target.cs
--- a/Engine/Source/UE4Game.Target.cs
+++ b/Engine/Source/UE4Game.Target.cs
@@ -109,11 +109,28 @@ public class UE4GameTarget : TargetRules
 	}
 	public override List<UnrealTargetPlatform> GUBP_GetPlatforms_MonolithicOnly(UnrealTargetPlatform HostPlatform)
     {
-        if (HostPlatform == UnrealTargetPlatform.Mac)
-        {
-            return new List<UnrealTargetPlatform> { HostPlatform, UnrealTargetPlatform.IOS };
-        }
-        return new List<UnrealTargetPlatform> { HostPlatform, UnrealTargetPlatform.Win32, UnrealTargetPlatform.IOS, UnrealTargetPlatform.XboxOne, UnrealTargetPlatform.PS4, UnrealTargetPlatform.Android, UnrealTargetPlatform.Linux, UnrealTargetPlatform.HTML5 };
+		List<UnrealTargetPlatform> Platforms = null;
+
+		switch(HostPlatform)
+		{
+			case UnrealTargetPlatform.Mac:
+				Platforms = new List<UnrealTargetPlatform> { HostPlatform, UnrealTargetPlatform.IOS };
+				break;
+
+			case UnrealTargetPlatform.Linux:
+				Platforms = new List<UnrealTargetPlatform> { HostPlatform };
+				break;
+
+			case UnrealTargetPlatform.Win64:
+				Platforms = new List<UnrealTargetPlatform> { HostPlatform, UnrealTargetPlatform.Win32, UnrealTargetPlatform.IOS, UnrealTargetPlatform.XboxOne, UnrealTargetPlatform.PS4, UnrealTargetPlatform.Android, UnrealTargetPlatform.Linux, UnrealTargetPlatform.HTML5 };
+				break;
+
+			default:
+				Platforms = new List<UnrealTargetPlatform>();
+				break;
+		}
+
+		return Platforms;
     }
     public override List<UnrealTargetConfiguration> GUBP_GetConfigs_MonolithicOnly(UnrealTargetPlatform HostPlatform, UnrealTargetPlatform Platform)
     {
diff --git a/Engine/Source/UE4Server.Target.cs b/Engine/Source/UE4Server.Target.cs
--- a/Engine/Source/UE4Server.Target.cs
+++ b/Engine/Source/UE4Server.Target.cs
@@ -58,11 +58,24 @@ public class UE4ServerTarget : TargetRules
 
     public override List<UnrealTargetPlatform> GUBP_GetPlatforms_MonolithicOnly(UnrealTargetPlatform HostPlatform)
     {
-        if (HostPlatform == UnrealTargetPlatform.Mac)
-        {
-            return new List<UnrealTargetPlatform>();
-        }
-        return new List<UnrealTargetPlatform> { HostPlatform, UnrealTargetPlatform.Win32, UnrealTargetPlatform.Linux };
+		List<UnrealTargetPlatform> Platforms = null;
+
+		switch(HostPlatform)
+		{
+			case UnrealTargetPlatform.Linux:
+				Platforms = new List<UnrealTargetPlatform> { HostPlatform };
+				break;
+
+			case UnrealTargetPlatform.Win64:
+				Platforms = new List<UnrealTargetPlatform> { HostPlatform, UnrealTargetPlatform.Win32, UnrealTargetPlatform.Linux };
+				break;
+
+			default:
+				Platforms = new List<UnrealTargetPlatform>();
+				break;
+		}
+
+		return Platforms;
     }
 
     public override List<UnrealTargetConfiguration> GUBP_GetConfigs_MonolithicOnly(UnrealTargetPlatform HostPlatform, UnrealTargetPlatform Platform)
