build script, add html5 tappy builds

diff --git a/Engine/Source/UE4Editor.Target.cs b/Engine/Source/UE4Editor.Target.cs
--- a/Engine/Source/UE4Editor.Target.cs
+++ b/Engine/Source/UE4Editor.Target.cs
@@ -105,9 +105,19 @@ public class UE4EditorTarget : TargetRules
             MobilePlats = new List<UnrealTargetPlatform> { HostPlatform, UnrealTargetPlatform.Android, UnrealTargetPlatform.IOS };
         }
 
+		List<UnrealTargetPlatform> TappyChickenPlats = null;
+		if (HostPlatform == UnrealTargetPlatform.Mac)
+		{
+			TappyChickenPlats = new List<UnrealTargetPlatform> { HostPlatform, UnrealTargetPlatform.IOS };
+		}
+		else
+		{
+			TappyChickenPlats = new List<UnrealTargetPlatform> { HostPlatform, UnrealTargetPlatform.Android, UnrealTargetPlatform.IOS, UnrealTargetPlatform.HTML5 };
+		}
+
         NonCodeProjectNames.Add("BlackJack", MobilePlats);
         NonCodeProjectNames.Add("MemoryGame", MobilePlats);
-		NonCodeProjectNames.Add("TappyChicken", MobilePlats);
+		NonCodeProjectNames.Add("TappyChicken", TappyChickenPlats);
         NonCodeProjectNames.Add("SwingNinja", MobilePlats);
         NonCodeProjectNames.Add("MobileTemple", MobilePlats);
 
@@ -132,7 +142,9 @@ public class UE4EditorTarget : TargetRules
                     new GUBPFormalBuild(UnrealTargetPlatform.Android, UnrealTargetConfiguration.Development, true),          
                     new GUBPFormalBuild(UnrealTargetPlatform.IOS, UnrealTargetConfiguration.Shipping, true),
                     new GUBPFormalBuild(UnrealTargetPlatform.IOS, UnrealTargetConfiguration.Test, true),
-                    new GUBPFormalBuild(UnrealTargetPlatform.IOS, UnrealTargetConfiguration.Development, true),          
+                    new GUBPFormalBuild(UnrealTargetPlatform.IOS, UnrealTargetConfiguration.Development, true),          			
+					new GUBPFormalBuild(UnrealTargetPlatform.HTML5, UnrealTargetConfiguration.Shipping, false),
+					new GUBPFormalBuild(UnrealTargetPlatform.HTML5, UnrealTargetConfiguration.Test, false),
             };
 
         var NonCodeProjectNames = new Dictionary<string, List<GUBPFormalBuild>>();
