Couple more renames for MobileTemple->SunTemple.

diff --git a/Engine/Source/UE4Editor.Target.cs b/Engine/Source/UE4Editor.Target.cs
--- a/Engine/Source/UE4Editor.Target.cs
+++ b/Engine/Source/UE4Editor.Target.cs
@@ -119,7 +119,7 @@ public class UE4EditorTarget : TargetRules
         NonCodeProjectNames.Add("MemoryGame", MobilePlats);
 		NonCodeProjectNames.Add("TappyChicken", TappyChickenPlats);
         NonCodeProjectNames.Add("SwingNinja", MobilePlats);
-        NonCodeProjectNames.Add("MobileTemple", MobilePlats);
+        NonCodeProjectNames.Add("SunTemple", MobilePlats);
 
         NonCodeProjectNames.Add("StarterContent", MobilePlats);
 		NonCodeProjectNames.Add("TP_2DSideScrollerBP", MobilePlats);
@@ -161,7 +161,7 @@ public class UE4EditorTarget : TargetRules
         NonCodeProjectNames.Add("BlackJack", MobileBuildSettings);
         NonCodeProjectNames.Add("MemoryGame", MobileBuildSettings);
         NonCodeProjectNames.Add("SwingNinja", MobileBuildSettings);
-        NonCodeProjectNames.Add("MobileTemple", MobileBuildSettings);
+        NonCodeProjectNames.Add("SunTemple", MobileBuildSettings);
 
         var PCBuildSettings = new List<GUBPFormalBuild>
             {
