UE4 project build still allowed for Win32 build trial and compilation fail later. Fixed.

diff --git a/Engine/Source/UE4Editor.Target.cs b/Engine/Source/UE4Editor.Target.cs
--- a/Engine/Source/UE4Editor.Target.cs
+++ b/Engine/Source/UE4Editor.Target.cs
@@ -10,11 +10,6 @@ public class UE4EditorTarget : TargetRules
 		Type = TargetType.Editor;
 	}
 
-	public override bool GetSupportedPlatforms(ref List<UnrealTargetPlatform> OutPlatforms)
-	{
-		return UnrealBuildTool.UnrealBuildTool.GetAllDesktopPlatforms(ref OutPlatforms, false);
-	}
-
 	public override void SetupBinaries(
 		TargetInfo Target,
 		ref List<UEBuildBinaryConfiguration> OutBuildBinaryConfigurations,
