Remove FriendsAndChat module from editor.

diff --git a/Engine/Source/UE4Editor.Target.cs b/Engine/Source/UE4Editor.Target.cs
--- a/Engine/Source/UE4Editor.Target.cs
+++ b/Engine/Source/UE4Editor.Target.cs
@@ -18,7 +18,6 @@ public class UE4EditorTarget : TargetRules
 	{
 		OutExtraModuleNames.Add("UE4Game");
 		OutExtraModuleNames.Add("GameMenuBuilder");
-		OutExtraModuleNames.Add("FriendsAndChat");
 		if ((Target.Platform == UnrealTargetPlatform.Win32) || (Target.Platform == UnrealTargetPlatform.Win64))
 		{
 			OutExtraModuleNames.Add("OnlineSubsystemNull");
