Adding templates and updating build configurations for non-code projects to allow CIS to build the sample configurations Engine QA needsfor testing. Reviewed by [REDACTED_PERSON]

diff --git a/Engine/Source/UE4Editor.Target.cs b/Engine/Source/UE4Editor.Target.cs
--- a/Engine/Source/UE4Editor.Target.cs
+++ b/Engine/Source/UE4Editor.Target.cs
@@ -60,17 +60,108 @@
     {
         var NonCodeProjectNames = new Dictionary<string, List<UnrealTargetPlatform>>();
 
+        #region Common Configurations
+
+        //DESKTOP PLATFORMS
         List<UnrealTargetPlatform> DesktopPlats = null;
         if (HostPlatform == UnrealTargetPlatform.Win64)
         {
-            DesktopPlats = new List<UnrealTargetPlatform> { HostPlatform, UnrealTargetPlatform.Linux };
+            DesktopPlats = new List<UnrealTargetPlatform> { HostPlatform, UnrealTargetPlatform.Linux, UnrealTargetPlatform.Win32 };
         }
         else
         {
             DesktopPlats = new List<UnrealTargetPlatform> { HostPlatform };
         }
+        //MOBILE PLATFORMS
+        List<UnrealTargetPlatform> MobilePlats = null;
+        if (HostPlatform == UnrealTargetPlatform.Mac)
+        {
+            MobilePlats = new List<UnrealTargetPlatform> { HostPlatform, UnrealTargetPlatform.IOS };
+        }
+        else
+        {
+            MobilePlats = new List<UnrealTargetPlatform> { HostPlatform, UnrealTargetPlatform.Android };
+        }
 
-        NonCodeProjectNames.Add("ElementalDemo", DesktopPlats);
+        //DESKTOP AND MOBILE PLATFORMS
+        List<UnrealTargetPlatform> DesktopAndMobilePlats = null;
+        if (HostPlatform == UnrealTargetPlatform.Mac)
+        {
+            DesktopAndMobilePlats = new List<UnrealTargetPlatform> { HostPlatform, UnrealTargetPlatform.IOS };
+        }
+        else
+        {
+            DesktopAndMobilePlats = new List<UnrealTargetPlatform> { HostPlatform, UnrealTargetPlatform.Android, UnrealTargetPlatform.Win32 };
+        }
+
+        //ALL SUPPORTED PLATFORMS
+        List<UnrealTargetPlatform> AllSupportedPlats = null;
+        if (HostPlatform == UnrealTargetPlatform.Mac)
+        {
+            AllSupportedPlats = new List<UnrealTargetPlatform> { HostPlatform, UnrealTargetPlatform.IOS };
+        }
+        else
+        {
+            AllSupportedPlats = new List<UnrealTargetPlatform> { HostPlatform, UnrealTargetPlatform.Android, UnrealTargetPlatform.HTML5, UnrealTargetPlatform.Win32, UnrealTargetPlatform.XboxOne, UnrealTargetPlatform.PS4 };
+        }
+        #endregion
+
+        #region Sample Specific Platform Configurations
+
+        //ELEMENTAL DEMO
+        List<UnrealTargetPlatform> ElementalDemoSupportedPlats = null;
+        if (HostPlatform == UnrealTargetPlatform.Mac)
+        {
+            ElementalDemoSupportedPlats = new List<UnrealTargetPlatform> { HostPlatform };
+        }
+        else
+        {
+            ElementalDemoSupportedPlats = new List<UnrealTargetPlatform> { HostPlatform, UnrealTargetPlatform.Win32, UnrealTargetPlatform.XboxOne, UnrealTargetPlatform.PS4 };
+        }
+
+        //TAPPY CHICKEN
+        List<UnrealTargetPlatform> TappyChickenPlats = null;
+        if (HostPlatform == UnrealTargetPlatform.Mac)
+        {
+            TappyChickenPlats = new List<UnrealTargetPlatform> { HostPlatform, UnrealTargetPlatform.IOS };
+        }
+        else
+        {
+            TappyChickenPlats = new List<UnrealTargetPlatform> { HostPlatform, UnrealTargetPlatform.Android, UnrealTargetPlatform.HTML5 };
+        }
+        //TP_PUZZLE
+        List<UnrealTargetPlatform> TPPuzzlePlats = null;
+        if (HostPlatform == UnrealTargetPlatform.Mac)
+        {
+            TPPuzzlePlats = new List<UnrealTargetPlatform> { HostPlatform, UnrealTargetPlatform.IOS };
+        }
+        else
+        {
+            TPPuzzlePlats = new List<UnrealTargetPlatform> { HostPlatform, UnrealTargetPlatform.Android, UnrealTargetPlatform.HTML5, UnrealTargetPlatform.Win32 };
+        }
+        //TP_TOPDOWN
+        List<UnrealTargetPlatform> TPTopDownPlats = null;
+        if (HostPlatform == UnrealTargetPlatform.Mac)
+        {
+            TPTopDownPlats = new List<UnrealTargetPlatform> { HostPlatform, UnrealTargetPlatform.IOS };
+        }
+        else
+        {
+            TPTopDownPlats = new List<UnrealTargetPlatform> { HostPlatform, UnrealTargetPlatform.Android, UnrealTargetPlatform.HTML5, UnrealTargetPlatform.Win32 };
+        }
+        //TP_TWINSTICK
+        List<UnrealTargetPlatform> TPTwinStickPlats = null;
+        if (HostPlatform == UnrealTargetPlatform.Mac)
+        {
+            TPTwinStickPlats = new List<UnrealTargetPlatform> { HostPlatform, UnrealTargetPlatform.IOS };
+        }
+        else
+        {
+            TPTwinStickPlats = new List<UnrealTargetPlatform> { HostPlatform, UnrealTargetPlatform.Android, UnrealTargetPlatform.Win32, UnrealTargetPlatform.XboxOne, UnrealTargetPlatform.PS4 };
+        }
+        #endregion
+
+        NonCodeProjectNames.Add("ElementalDemo", ElementalDemoSupportedPlats);
         //NonCodeProjectNames.Add("InfiltratorDemo", DesktopPlats);
         NonCodeProjectNames.Add("HoverShip", DesktopPlats);
         NonCodeProjectNames.Add("BlueprintOffice", DesktopPlats);
@@ -91,46 +182,96 @@
         NonCodeProjectNames.Add("PostProcessMatinee", DesktopPlats);
         NonCodeProjectNames.Add("SciFiHallway", DesktopPlats);
 
-        List<UnrealTargetPlatform> MobilePlats = null;
-        if (HostPlatform == UnrealTargetPlatform.Mac)
-        {
-            MobilePlats = new List<UnrealTargetPlatform> { HostPlatform, UnrealTargetPlatform.IOS };
-        }
-        else
-        {
-            MobilePlats = new List<UnrealTargetPlatform> { HostPlatform, UnrealTargetPlatform.Android, UnrealTargetPlatform.IOS };
-        }
+        NonCodeProjectNames.Add("BlackJack", DesktopAndMobilePlats);
+        NonCodeProjectNames.Add("MemoryGame", DesktopAndMobilePlats);
+        NonCodeProjectNames.Add("TappyChicken", DesktopAndMobilePlats);
+        NonCodeProjectNames.Add("SwingNinja", DesktopAndMobilePlats);
+        NonCodeProjectNames.Add("SunTemple", DesktopAndMobilePlats);
 
-		List<UnrealTargetPlatform> TappyChickenPlats = null;
-		if (HostPlatform == UnrealTargetPlatform.Mac)
-		{
-			TappyChickenPlats = new List<UnrealTargetPlatform> { HostPlatform, UnrealTargetPlatform.IOS };
-		}
-		else
-		{
-			TappyChickenPlats = new List<UnrealTargetPlatform> { HostPlatform, UnrealTargetPlatform.Android, UnrealTargetPlatform.IOS, UnrealTargetPlatform.HTML5 };
-		}
+        NonCodeProjectNames.Add("TP_2DSideScrollerBP", AllSupportedPlats);
+        NonCodeProjectNames.Add("TP_FirstPersonBP", AllSupportedPlats);
+        NonCodeProjectNames.Add("TP_FlyingBP", AllSupportedPlats);
+        NonCodeProjectNames.Add("TP_PuzzleBP", TPPuzzlePlats);
+        NonCodeProjectNames.Add("TP_RollingBP", AllSupportedPlats);
+        NonCodeProjectNames.Add("TP_SideScrollerBP", AllSupportedPlats);
+        NonCodeProjectNames.Add("TP_ThirdPersonBP", AllSupportedPlats);
+        NonCodeProjectNames.Add("TP_TopDownBP", TPTopDownPlats);
+        NonCodeProjectNames.Add("TP_TwinStickBP", TPTwinStickPlats);
+        NonCodeProjectNames.Add("TP_VehicleBP", AllSupportedPlats);
+        NonCodeProjectNames.Add("TP_VehicleAdvBP", AllSupportedPlats);
 
-        NonCodeProjectNames.Add("BlackJack", MobilePlats);
-        NonCodeProjectNames.Add("MemoryGame", MobilePlats);
-		NonCodeProjectNames.Add("TappyChicken", TappyChickenPlats);
-        NonCodeProjectNames.Add("SwingNinja", MobilePlats);
-        NonCodeProjectNames.Add("SunTemple", MobilePlats);
-
-        NonCodeProjectNames.Add("StarterContent", MobilePlats);
-		NonCodeProjectNames.Add("TP_2DSideScrollerBP", MobilePlats);
-        NonCodeProjectNames.Add("TP_FirstPersonBP", MobilePlats);
-        NonCodeProjectNames.Add("TP_FlyingBP", MobilePlats);
-        NonCodeProjectNames.Add("TP_RollingBP", MobilePlats);
-        NonCodeProjectNames.Add("TP_SideScrollerBP", MobilePlats);
-        NonCodeProjectNames.Add("TP_ThirdPersonBP", MobilePlats);
-        NonCodeProjectNames.Add("TP_TopDownBP", MobilePlats);
-        NonCodeProjectNames.Add("TP_VehicleBP", MobilePlats);
-
         return NonCodeProjectNames;
     }
     public override Dictionary<string, List<GUBPFormalBuild>> GUBP_GetNonCodeFormalBuilds_BaseEditorTypeOnly()
     {
+        #region Sample Configurations
+        //Generic Configurations
+        var PCBuildSettings = new List<GUBPFormalBuild>
+            {
+                    new GUBPFormalBuild(UnrealTargetPlatform.Win32, UnrealTargetConfiguration.Test, false),
+                    new GUBPFormalBuild(UnrealTargetPlatform.Win64, UnrealTargetConfiguration.Test, false),
+                    new GUBPFormalBuild(UnrealTargetPlatform.Mac, UnrealTargetConfiguration.Test, false),
+            };
+
+        var MobileBuildSettings = new List<GUBPFormalBuild>
+            {
+                    new GUBPFormalBuild(UnrealTargetPlatform.Android, UnrealTargetConfiguration.Test, false),
+                    new GUBPFormalBuild(UnrealTargetPlatform.IOS, UnrealTargetConfiguration.Test, false),
+            };
+        // Non-Code Sample Configurations
+        var BlackJackBuildSettings = new List<GUBPFormalBuild>
+            {
+                    new GUBPFormalBuild(UnrealTargetPlatform.Win32, UnrealTargetConfiguration.Test, false),
+                    new GUBPFormalBuild(UnrealTargetPlatform.Win64, UnrealTargetConfiguration.Test, false),
+                    new GUBPFormalBuild(UnrealTargetPlatform.Mac, UnrealTargetConfiguration.Test, false),
+                    new GUBPFormalBuild(UnrealTargetPlatform.Android, UnrealTargetConfiguration.Test, false),
+                    new GUBPFormalBuild(UnrealTargetPlatform.IOS, UnrealTargetConfiguration.Test, false),
+            };
+        
+        var MemoryGameBuildSettings = new List<GUBPFormalBuild>
+            {
+                    new GUBPFormalBuild(UnrealTargetPlatform.Win32, UnrealTargetConfiguration.Test, false),
+                    new GUBPFormalBuild(UnrealTargetPlatform.Win64, UnrealTargetConfiguration.Test, false),
+                    new GUBPFormalBuild(UnrealTargetPlatform.Mac, UnrealTargetConfiguration.Test, false),
+                    new GUBPFormalBuild(UnrealTargetPlatform.Android, UnrealTargetConfiguration.Test, false),
+                    new GUBPFormalBuild(UnrealTargetPlatform.IOS, UnrealTargetConfiguration.Test, false),
+                    new GUBPFormalBuild(UnrealTargetPlatform.HTML5, UnrealTargetConfiguration.Test, false),
+            };
+        
+        var SwingNinjaBuildSettings = new List<GUBPFormalBuild>
+            {
+                    new GUBPFormalBuild(UnrealTargetPlatform.Win32, UnrealTargetConfiguration.Test, false),
+                    new GUBPFormalBuild(UnrealTargetPlatform.Win64, UnrealTargetConfiguration.Test, false),
+                    new GUBPFormalBuild(UnrealTargetPlatform.Mac, UnrealTargetConfiguration.Test, false),
+                    new GUBPFormalBuild(UnrealTargetPlatform.Android, UnrealTargetConfiguration.Test, false),
+                    new GUBPFormalBuild(UnrealTargetPlatform.IOS, UnrealTargetConfiguration.Test, false),
+            };
+        
+        var SunTempleBuildSettings = new List<GUBPFormalBuild>
+            {
+                    new GUBPFormalBuild(UnrealTargetPlatform.Win32, UnrealTargetConfiguration.Test, false),
+                    new GUBPFormalBuild(UnrealTargetPlatform.Win64, UnrealTargetConfiguration.Test, false),
+                    new GUBPFormalBuild(UnrealTargetPlatform.Mac, UnrealTargetConfiguration.Test, false),
+                    new GUBPFormalBuild(UnrealTargetPlatform.Android, UnrealTargetConfiguration.Test, false),
+                    new GUBPFormalBuild(UnrealTargetPlatform.IOS, UnrealTargetConfiguration.Test, false),
+                    new GUBPFormalBuild(UnrealTargetPlatform.HTML5, UnrealTargetConfiguration.Test, false),
+            };
+
+        var ElementalBuildSettings = new List<GUBPFormalBuild>
+            {
+                    new GUBPFormalBuild(UnrealTargetPlatform.Win32, UnrealTargetConfiguration.Test, false),
+                    new GUBPFormalBuild(UnrealTargetPlatform.Win64, UnrealTargetConfiguration.Test, false),
+                    new GUBPFormalBuild(UnrealTargetPlatform.Mac, UnrealTargetConfiguration.Test, false),
+                    new GUBPFormalBuild(UnrealTargetPlatform.XboxOne, UnrealTargetConfiguration.Test, false),
+                    new GUBPFormalBuild(UnrealTargetPlatform.PS4, UnrealTargetConfiguration.Test, false),
+                    
+                    new GUBPFormalBuild(UnrealTargetPlatform.XboxOne, UnrealTargetConfiguration.Shipping, false),
+                    new GUBPFormalBuild(UnrealTargetPlatform.PS4, UnrealTargetConfiguration.Shipping, false),
+                    
+                    new GUBPFormalBuild(UnrealTargetPlatform.XboxOne, UnrealTargetConfiguration.Development, false),
+                    new GUBPFormalBuild(UnrealTargetPlatform.PS4, UnrealTargetConfiguration.Development, false),
+            };
+
         var TappyChickenBuildSettings = new List<GUBPFormalBuild>
             {
                     new GUBPFormalBuild(UnrealTargetPlatform.Android, UnrealTargetConfiguration.Shipping, true),
@@ -142,30 +283,66 @@
 					new GUBPFormalBuild(UnrealTargetPlatform.HTML5, UnrealTargetConfiguration.Shipping, false),
 					new GUBPFormalBuild(UnrealTargetPlatform.HTML5, UnrealTargetConfiguration.Test, false),
             };
+        #endregion
 
+        #region Template Configurations
+        // Template Build Configurations
+        // DEFAULT 
+        var AllBuildSettings = new List<GUBPFormalBuild>
+            {                                     
+                    new GUBPFormalBuild(UnrealTargetPlatform.IOS, UnrealTargetConfiguration.Test, false),
+                    new GUBPFormalBuild(UnrealTargetPlatform.Android, UnrealTargetConfiguration.Test, false),
+                    new GUBPFormalBuild(UnrealTargetPlatform.HTML5, UnrealTargetConfiguration.Test, false),
+                    new GUBPFormalBuild(UnrealTargetPlatform.Win32, UnrealTargetConfiguration.Test, false),
+                    new GUBPFormalBuild(UnrealTargetPlatform.Win64, UnrealTargetConfiguration.Test, false),
+                    new GUBPFormalBuild(UnrealTargetPlatform.Mac, UnrealTargetConfiguration.Test, false),
+                    new GUBPFormalBuild(UnrealTargetPlatform.XboxOne, UnrealTargetConfiguration.Test, false),
+                    new GUBPFormalBuild(UnrealTargetPlatform.PS4, UnrealTargetConfiguration.Test, false),
+            };
+
+        //TP_PUZZLE
+        var TP_PuzzleBuildSettings = new List<GUBPFormalBuild>
+            {
+                    new GUBPFormalBuild(UnrealTargetPlatform.IOS, UnrealTargetConfiguration.Test, false),
+                    new GUBPFormalBuild(UnrealTargetPlatform.Android, UnrealTargetConfiguration.Test, false),
+                    new GUBPFormalBuild(UnrealTargetPlatform.HTML5, UnrealTargetConfiguration.Test, false),
+                    new GUBPFormalBuild(UnrealTargetPlatform.Win32, UnrealTargetConfiguration.Test, false),
+                    new GUBPFormalBuild(UnrealTargetPlatform.Win64, UnrealTargetConfiguration.Test, false),
+                    new GUBPFormalBuild(UnrealTargetPlatform.Mac, UnrealTargetConfiguration.Test, false),
+            };
+
+        //TP_TOPDOWN
+        var TP_TopDownBuildSettings = new List<GUBPFormalBuild>
+            {
+                    new GUBPFormalBuild(UnrealTargetPlatform.IOS, UnrealTargetConfiguration.Test, false),
+                    new GUBPFormalBuild(UnrealTargetPlatform.Android, UnrealTargetConfiguration.Test, false),
+                    new GUBPFormalBuild(UnrealTargetPlatform.HTML5, UnrealTargetConfiguration.Test, false),
+                    new GUBPFormalBuild(UnrealTargetPlatform.Win32, UnrealTargetConfiguration.Test, false),
+                    new GUBPFormalBuild(UnrealTargetPlatform.Win64, UnrealTargetConfiguration.Test, false),
+                    new GUBPFormalBuild(UnrealTargetPlatform.Mac, UnrealTargetConfiguration.Test, false),
+            };
+
+        //TP_TWINSTICK
+        var TP_TwinStickBuildSettings = new List<GUBPFormalBuild>
+            {
+                    new GUBPFormalBuild(UnrealTargetPlatform.IOS, UnrealTargetConfiguration.Test, false),
+                    new GUBPFormalBuild(UnrealTargetPlatform.Android, UnrealTargetConfiguration.Test, false),
+                    new GUBPFormalBuild(UnrealTargetPlatform.Win32, UnrealTargetConfiguration.Test, false),
+                    new GUBPFormalBuild(UnrealTargetPlatform.Win64, UnrealTargetConfiguration.Test, false),
+                    new GUBPFormalBuild(UnrealTargetPlatform.Mac, UnrealTargetConfiguration.Test, false),
+                    new GUBPFormalBuild(UnrealTargetPlatform.XboxOne, UnrealTargetConfiguration.Test, false),
+                    new GUBPFormalBuild(UnrealTargetPlatform.PS4, UnrealTargetConfiguration.Test, false),
+            };
+        #endregion
+
         var NonCodeProjectNames = new Dictionary<string, List<GUBPFormalBuild>>();
-        NonCodeProjectNames.Add("TappyChicken", TappyChickenBuildSettings);
 
-
-
-        var MobileBuildSettings = new List<GUBPFormalBuild>
-            {
-                    new GUBPFormalBuild(UnrealTargetPlatform.Android, UnrealTargetConfiguration.Test, false),
-                    new GUBPFormalBuild(UnrealTargetPlatform.IOS, UnrealTargetConfiguration.Test, false),
-            };
-
-        NonCodeProjectNames.Add("BlackJack", MobileBuildSettings);
-        NonCodeProjectNames.Add("MemoryGame", MobileBuildSettings);
-        NonCodeProjectNames.Add("SwingNinja", MobileBuildSettings);
-        NonCodeProjectNames.Add("SunTemple", MobileBuildSettings);
-
-        var PCBuildSettings = new List<GUBPFormalBuild>
-            {
-                    new GUBPFormalBuild(UnrealTargetPlatform.Win32, UnrealTargetConfiguration.Test, false),
-                    new GUBPFormalBuild(UnrealTargetPlatform.Win64, UnrealTargetConfiguration.Test, false),
-                    new GUBPFormalBuild(UnrealTargetPlatform.Mac, UnrealTargetConfiguration.Test, false),
-            };
-
+        //Add Samples to the list with its corresponding settings
+        NonCodeProjectNames.Add("TappyChicken", TappyChickenBuildSettings);    
+        NonCodeProjectNames.Add("BlackJack", BlackJackBuildSettings);
+        NonCodeProjectNames.Add("MemoryGame", MemoryGameBuildSettings);
+        NonCodeProjectNames.Add("SwingNinja", SwingNinjaBuildSettings);
+        NonCodeProjectNames.Add("SunTemple", SunTempleBuildSettings);
         NonCodeProjectNames.Add("BlueprintOffice", PCBuildSettings);
         NonCodeProjectNames.Add("ContentExamples", PCBuildSettings);
         NonCodeProjectNames.Add("CouchKnights", PCBuildSettings);
@@ -176,26 +353,20 @@
         NonCodeProjectNames.Add("RealisticRendering", PCBuildSettings);
         NonCodeProjectNames.Add("ReflectionsSubway", PCBuildSettings);
         NonCodeProjectNames.Add("Stylized", PCBuildSettings);
-
-
-        var ElementalBuildSettings = new List<GUBPFormalBuild>
-            {
-                    new GUBPFormalBuild(UnrealTargetPlatform.Win32, UnrealTargetConfiguration.Test, false),
-                    new GUBPFormalBuild(UnrealTargetPlatform.Win64, UnrealTargetConfiguration.Test, false),
-                    new GUBPFormalBuild(UnrealTargetPlatform.Mac, UnrealTargetConfiguration.Test, false),
-                    new GUBPFormalBuild(UnrealTargetPlatform.XboxOne, UnrealTargetConfiguration.Test, false),
-                    new GUBPFormalBuild(UnrealTargetPlatform.PS4, UnrealTargetConfiguration.Test, false),
-                    
-                    new GUBPFormalBuild(UnrealTargetPlatform.XboxOne, UnrealTargetConfiguration.Shipping, false),
-                    new GUBPFormalBuild(UnrealTargetPlatform.PS4, UnrealTargetConfiguration.Shipping, false),
-                    
-                    new GUBPFormalBuild(UnrealTargetPlatform.XboxOne, UnrealTargetConfiguration.Development, false),
-                    new GUBPFormalBuild(UnrealTargetPlatform.PS4, UnrealTargetConfiguration.Development, false),
-            };
-
-
         NonCodeProjectNames.Add("ElementalDemo", ElementalBuildSettings);
 
+        //Add Templates to the list with its corresponding settings
+        NonCodeProjectNames.Add("TP_2DSideScrollerBP", AllBuildSettings);
+        NonCodeProjectNames.Add("TP_FirstPersonBP", AllBuildSettings);
+        NonCodeProjectNames.Add("TP_FlyingBP", AllBuildSettings);
+        NonCodeProjectNames.Add("TP_PuzzleBP", TP_PuzzleBuildSettings);
+        NonCodeProjectNames.Add("TP_RollingBP", AllBuildSettings);
+        NonCodeProjectNames.Add("TP_SideScrollerBP", AllBuildSettings);
+        NonCodeProjectNames.Add("TP_ThirdPersonBP", AllBuildSettings);
+        NonCodeProjectNames.Add("TP_TopDownBP", TP_TopDownBuildSettings);
+        NonCodeProjectNames.Add("TP_TwinStickBP", TP_TwinStickBuildSettings);
+        NonCodeProjectNames.Add("TP_VehicleBP", AllBuildSettings);
+        NonCodeProjectNames.Add("TP_VehicleAdvBP", AllBuildSettings);
 
 
         return NonCodeProjectNames;
