build script, comment out unused lists to prevent mac failure

diff --git a/Engine/Source/UE4Editor.Target.cs b/Engine/Source/UE4Editor.Target.cs
--- a/Engine/Source/UE4Editor.Target.cs
+++ b/Engine/Source/UE4Editor.Target.cs
@@ -73,15 +73,15 @@ public class UE4EditorTarget : TargetRules
             DesktopPlats = new List<UnrealTargetPlatform> { HostPlatform };
         }
         //MOBILE PLATFORMS
-        List<UnrealTargetPlatform> MobilePlats = null;
-        if (HostPlatform == UnrealTargetPlatform.Mac)
-        {
-            MobilePlats = new List<UnrealTargetPlatform> { HostPlatform, UnrealTargetPlatform.IOS };
-        }
-        else
-        {
-            MobilePlats = new List<UnrealTargetPlatform> { HostPlatform, UnrealTargetPlatform.Android };
-        }
+		//List<UnrealTargetPlatform> MobilePlats = null;
+		//if (HostPlatform == UnrealTargetPlatform.Mac)
+		//{
+		//	MobilePlats = new List<UnrealTargetPlatform> { HostPlatform, UnrealTargetPlatform.IOS };
+		//}
+		//else
+		//{
+		//	MobilePlats = new List<UnrealTargetPlatform> { HostPlatform, UnrealTargetPlatform.Android };
+		//}
 
         //DESKTOP AND MOBILE PLATFORMS
         List<UnrealTargetPlatform> DesktopAndMobilePlats = null;
@@ -120,15 +120,15 @@ public class UE4EditorTarget : TargetRules
         }
 
         //TAPPY CHICKEN
-        List<UnrealTargetPlatform> TappyChickenPlats = null;
-        if (HostPlatform == UnrealTargetPlatform.Mac)
-        {
-            TappyChickenPlats = new List<UnrealTargetPlatform> { HostPlatform, UnrealTargetPlatform.IOS };
-        }
-        else
-        {
-            TappyChickenPlats = new List<UnrealTargetPlatform> { HostPlatform, UnrealTargetPlatform.Android, UnrealTargetPlatform.HTML5 };
-        }
+		//List<UnrealTargetPlatform> TappyChickenPlats = null;
+		//if (HostPlatform == UnrealTargetPlatform.Mac)
+		//{
+		//	TappyChickenPlats = new List<UnrealTargetPlatform> { HostPlatform, UnrealTargetPlatform.IOS };
+		//}
+		//else
+		//{
+		//	TappyChickenPlats = new List<UnrealTargetPlatform> { HostPlatform, UnrealTargetPlatform.Android, UnrealTargetPlatform.HTML5 };
+		//}
         //TP_PUZZLE
         List<UnrealTargetPlatform> TPPuzzlePlats = null;
         if (HostPlatform == UnrealTargetPlatform.Mac)
@@ -213,11 +213,11 @@ public class UE4EditorTarget : TargetRules
                     new GUBPFormalBuild(UnrealTargetPlatform.Mac, UnrealTargetConfiguration.Test, false),
             };
 
-        var MobileBuildSettings = new List<GUBPFormalBuild>
-            {
-                    new GUBPFormalBuild(UnrealTargetPlatform.Android, UnrealTargetConfiguration.Test, false),
-                    new GUBPFormalBuild(UnrealTargetPlatform.IOS, UnrealTargetConfiguration.Test, false),
-            };
+		//var MobileBuildSettings = new List<GUBPFormalBuild>
+		//	{
+		//			new GUBPFormalBuild(UnrealTargetPlatform.Android, UnrealTargetConfiguration.Test, false),
+		//			new GUBPFormalBuild(UnrealTargetPlatform.IOS, UnrealTargetConfiguration.Test, false),
+		//	};
         // Non-Code Sample Configurations
         var BlackJackBuildSettings = new List<GUBPFormalBuild>
             {
