Updating TappyChicken config to use the TappyChickenPlats configuration. This will now allow tappy to build for HTML5. It was incorrectly set to a different DesktopAndMobilePlats which does not include HTML5.

diff --git a/Engine/Source/UE4Editor.Target.cs b/Engine/Source/UE4Editor.Target.cs
--- a/Engine/Source/UE4Editor.Target.cs
+++ b/Engine/Source/UE4Editor.Target.cs
@@ -120,15 +120,15 @@ public class UE4EditorTarget : TargetRules
         }
 
         //TAPPY CHICKEN
-		//List<UnrealTargetPlatform> TappyChickenPlats = null;
-		//if (HostPlatform == UnrealTargetPlatform.Mac)
-		//{
-		//	TappyChickenPlats = new List<UnrealTargetPlatform> { HostPlatform, UnrealTargetPlatform.IOS };
-		//}
-		//else
-		//{
-		//	TappyChickenPlats = new List<UnrealTargetPlatform> { HostPlatform, UnrealTargetPlatform.Android, UnrealTargetPlatform.HTML5 };
-		//}
+        List<UnrealTargetPlatform> TappyChickenPlats = null;
+        if (HostPlatform == UnrealTargetPlatform.Mac)
+        {
+            TappyChickenPlats = new List<UnrealTargetPlatform> { HostPlatform, UnrealTargetPlatform.IOS };
+        }
+        else
+        {
+            TappyChickenPlats = new List<UnrealTargetPlatform> { HostPlatform, UnrealTargetPlatform.Android, UnrealTargetPlatform.HTML5 };
+        }
         //TP_PUZZLE
         List<UnrealTargetPlatform> TPPuzzlePlats = null;
         if (HostPlatform == UnrealTargetPlatform.Mac)
@@ -184,7 +184,7 @@ public class UE4EditorTarget : TargetRules
 
         NonCodeProjectNames.Add("BlackJack", DesktopAndMobilePlats);
         NonCodeProjectNames.Add("MemoryGame", DesktopAndMobilePlats);
-        NonCodeProjectNames.Add("TappyChicken", DesktopAndMobilePlats);
+        NonCodeProjectNames.Add("TappyChicken", TappyChickenPlats);
         NonCodeProjectNames.Add("SwingNinja", DesktopAndMobilePlats);
         NonCodeProjectNames.Add("SunTemple", DesktopAndMobilePlats);
 
