Include the optional runtime WebBrowser module for Rocket on Win64 and Mac. Win32 to come.

diff --git a/Engine/Source/UE4Game.Target.cs b/Engine/Source/UE4Game.Target.cs
--- a/Engine/Source/UE4Game.Target.cs
+++ b/Engine/Source/UE4Game.Target.cs
@@ -34,7 +34,13 @@ public class UE4GameTarget : TargetRules
 		if (UnrealBuildTool.UnrealBuildTool.BuildingRocket())
 		{
 			OutExtraModuleNames.Add("GameMenuBuilder");
-			if ((Target.Platform == UnrealTargetPlatform.Win32) || (Target.Platform == UnrealTargetPlatform.Win64))
+
+			if (Target.Platform == UnrealTargetPlatform.Win64 || Target.Platform == UnrealTargetPlatform.Mac)
+			{
+				OutExtraModuleNames.Add("WebBrowser");
+			}
+
+			if (Target.Platform == UnrealTargetPlatform.Win32 || Target.Platform == UnrealTargetPlatform.Win64)
 			{
 				OutExtraModuleNames.Add("OnlineSubsystemNull");
 				OutExtraModuleNames.Add("OnlineSubsystemAmazon");
