Disabled WebBrowser dependency for Rocket games on Mac to temporarily work around UE-7989

diff --git a/Engine/Source/UE4Game.Target.cs b/Engine/Source/UE4Game.Target.cs
--- a/Engine/Source/UE4Game.Target.cs
+++ b/Engine/Source/UE4Game.Target.cs
@@ -35,7 +35,7 @@ public class UE4GameTarget : TargetRules
 		{
 			OutExtraModuleNames.Add("GameMenuBuilder");
 
-			if (Target.Platform == UnrealTargetPlatform.Win32 || Target.Platform == UnrealTargetPlatform.Win64 || Target.Platform == UnrealTargetPlatform.Mac)
+			if (Target.Platform == UnrealTargetPlatform.Win32 || Target.Platform == UnrealTargetPlatform.Win64)
 			{
 				OutExtraModuleNames.Add("WebBrowser");
 			}
