Add FP_FirstPersonBP to list of non-code samples.

diff --git a/Engine/Source/UE4Editor.Target.cs b/Engine/Source/UE4Editor.Target.cs
--- a/Engine/Source/UE4Editor.Target.cs
+++ b/Engine/Source/UE4Editor.Target.cs
@@ -188,6 +188,7 @@ public class UE4EditorTarget : TargetRules
         NonCodeProjectNames.Add("SwingNinja", DesktopAndMobilePlats);
         NonCodeProjectNames.Add("SunTemple", DesktopAndMobilePlats);
 
+        NonCodeProjectNames.Add("FP_FirstPersonBP", AllSupportedPlats);
         NonCodeProjectNames.Add("TP_2DSideScrollerBP", AllSupportedPlats);
         NonCodeProjectNames.Add("TP_FirstPersonBP", AllSupportedPlats);
         NonCodeProjectNames.Add("TP_FlyingBP", AllSupportedPlats);
@@ -356,6 +357,7 @@ public class UE4EditorTarget : TargetRules
         NonCodeProjectNames.Add("ElementalDemo", ElementalBuildSettings);
 
         //Add Templates to the list with its corresponding settings
+        NonCodeProjectNames.Add("FP_FirstPersonBP", AllBuildSettings);
         NonCodeProjectNames.Add("TP_2DSideScrollerBP", AllBuildSettings);
         NonCodeProjectNames.Add("TP_FirstPersonBP", AllBuildSettings);
         NonCodeProjectNames.Add("TP_FlyingBP", AllBuildSettings);
