build script, temp remove tappy from ue4 builds

diff --git a/Engine/Source/UE4Editor.Target.cs b/Engine/Source/UE4Editor.Target.cs
--- a/Engine/Source/UE4Editor.Target.cs
+++ b/Engine/Source/UE4Editor.Target.cs
@@ -120,7 +120,7 @@ public class UE4EditorTarget : TargetRules
         }
 
         //TAPPY CHICKEN
-        List<UnrealTargetPlatform> TappyChickenPlats = null;
+        /*List<UnrealTargetPlatform> TappyChickenPlats = null;
         if (HostPlatform == UnrealTargetPlatform.Mac)
         {
             TappyChickenPlats = new List<UnrealTargetPlatform> { HostPlatform, UnrealTargetPlatform.IOS };
@@ -128,7 +128,7 @@ public class UE4EditorTarget : TargetRules
         else
         {
             TappyChickenPlats = new List<UnrealTargetPlatform> { HostPlatform, UnrealTargetPlatform.Android, UnrealTargetPlatform.HTML5 };
-        }
+        }*/
         //TP_PUZZLE
         List<UnrealTargetPlatform> TPPuzzlePlats = null;
         if (HostPlatform == UnrealTargetPlatform.Mac)
@@ -184,7 +184,7 @@ public class UE4EditorTarget : TargetRules
 
         NonCodeProjectNames.Add("BlackJack", DesktopAndMobilePlats);
         NonCodeProjectNames.Add("MemoryGame", DesktopAndMobilePlats);
-        NonCodeProjectNames.Add("TappyChicken", TappyChickenPlats);
+        //NonCodeProjectNames.Add("TappyChicken", TappyChickenPlats);
         NonCodeProjectNames.Add("SwingNinja", DesktopAndMobilePlats);
         NonCodeProjectNames.Add("SunTemple", DesktopAndMobilePlats);
 
@@ -273,7 +273,7 @@ public class UE4EditorTarget : TargetRules
                     new GUBPFormalBuild(UnrealTargetPlatform.PS4, UnrealTargetConfiguration.Development, false),
             };
 
-        var TappyChickenBuildSettings = new List<GUBPFormalBuild>
+        /*var TappyChickenBuildSettings = new List<GUBPFormalBuild>
             {
                     new GUBPFormalBuild(UnrealTargetPlatform.Android, UnrealTargetConfiguration.Shipping, true),
                     new GUBPFormalBuild(UnrealTargetPlatform.Android, UnrealTargetConfiguration.Test, true),                    
@@ -283,7 +283,7 @@ public class UE4EditorTarget : TargetRules
                     new GUBPFormalBuild(UnrealTargetPlatform.IOS, UnrealTargetConfiguration.Development, true),          			
 					new GUBPFormalBuild(UnrealTargetPlatform.HTML5, UnrealTargetConfiguration.Shipping, false),
 					new GUBPFormalBuild(UnrealTargetPlatform.HTML5, UnrealTargetConfiguration.Test, false),
-            };
+            };*/
         #endregion
 
         #region Template Configurations
@@ -339,7 +339,7 @@ public class UE4EditorTarget : TargetRules
         var NonCodeProjectNames = new Dictionary<string, List<GUBPFormalBuild>>();
 
         //Add Samples to the list with its corresponding settings
-        NonCodeProjectNames.Add("TappyChicken", TappyChickenBuildSettings);    
+        //NonCodeProjectNames.Add("TappyChicken", TappyChickenBuildSettings);    
         NonCodeProjectNames.Add("BlackJack", BlackJackBuildSettings);
         NonCodeProjectNames.Add("MemoryGame", MemoryGameBuildSettings);
         NonCodeProjectNames.Add("SwingNinja", SwingNinjaBuildSettings);
