build script, fix tp_2dsidescroller noncode plats

diff --git a/Engine/Source/UE4Editor.Target.cs b/Engine/Source/UE4Editor.Target.cs
--- a/Engine/Source/UE4Editor.Target.cs
+++ b/Engine/Source/UE4Editor.Target.cs
@@ -129,6 +129,16 @@ public class UE4EditorTarget : TargetRules
         {
             TappyChickenPlats = new List<UnrealTargetPlatform> { HostPlatform, UnrealTargetPlatform.Android, UnrealTargetPlatform.HTML5 };
         }
+
+		List<UnrealTargetPlatform> TP_2DSideScrollerPlats = null;
+		if (HostPlatform == UnrealTargetPlatform.Mac)
+        {
+			TP_2DSideScrollerPlats = new List<UnrealTargetPlatform> { HostPlatform };
+        }
+        else
+        {
+			TP_2DSideScrollerPlats = new List<UnrealTargetPlatform> { HostPlatform, UnrealTargetPlatform.Android, UnrealTargetPlatform.HTML5, UnrealTargetPlatform.Win32, UnrealTargetPlatform.XboxOne, UnrealTargetPlatform.PS4 };
+        }
         //TP_PUZZLE
         List<UnrealTargetPlatform> TPPuzzlePlats = null;
         if (HostPlatform == UnrealTargetPlatform.Mac)
@@ -189,7 +199,7 @@ public class UE4EditorTarget : TargetRules
         NonCodeProjectNames.Add("SunTemple", DesktopAndMobilePlats);
 
         NonCodeProjectNames.Add("FP_FirstPersonBP", AllSupportedPlats);
-        NonCodeProjectNames.Add("TP_2DSideScrollerBP", AllSupportedPlats);
+		NonCodeProjectNames.Add("TP_2DSideScrollerBP", TP_2DSideScrollerPlats);
         NonCodeProjectNames.Add("TP_FirstPersonBP", AllSupportedPlats);
         NonCodeProjectNames.Add("TP_FlyingBP", AllSupportedPlats);
         NonCodeProjectNames.Add("TP_PuzzleBP", TPPuzzlePlats);
@@ -297,6 +307,16 @@ public class UE4EditorTarget : TargetRules
                     new GUBPFormalBuild(UnrealTargetPlatform.XboxOne, UnrealTargetConfiguration.Test, false),
                     new GUBPFormalBuild(UnrealTargetPlatform.PS4, UnrealTargetConfiguration.Test, false),
             };
+		var TP_2DSideScrollerBuildSettings = new List<GUBPFormalBuild>
+            {                    
+                    new GUBPFormalBuild(UnrealTargetPlatform.Android, UnrealTargetConfiguration.Test, false),
+                    new GUBPFormalBuild(UnrealTargetPlatform.HTML5, UnrealTargetConfiguration.Test, false),
+                    new GUBPFormalBuild(UnrealTargetPlatform.Win32, UnrealTargetConfiguration.Test, false),
+                    new GUBPFormalBuild(UnrealTargetPlatform.Win64, UnrealTargetConfiguration.Test, false),
+                    new GUBPFormalBuild(UnrealTargetPlatform.Mac, UnrealTargetConfiguration.Test, false),
+                    new GUBPFormalBuild(UnrealTargetPlatform.XboxOne, UnrealTargetConfiguration.Test, false),
+                    new GUBPFormalBuild(UnrealTargetPlatform.PS4, UnrealTargetConfiguration.Test, false),
+            };
 
         //TP_PUZZLE
         var TP_PuzzleBuildSettings = new List<GUBPFormalBuild>
@@ -355,7 +375,7 @@ public class UE4EditorTarget : TargetRules
 
         //Add Templates to the list with its corresponding settings
         NonCodeProjectNames.Add("FP_FirstPersonBP", AllBuildSettings);
-        NonCodeProjectNames.Add("TP_2DSideScrollerBP", AllBuildSettings);
+		NonCodeProjectNames.Add("TP_2DSideScrollerBP", TP_2DSideScrollerBuildSettings);
         NonCodeProjectNames.Add("TP_FirstPersonBP", AllBuildSettings);
         NonCodeProjectNames.Add("TP_FlyingBP", AllBuildSettings);
         NonCodeProjectNames.Add("TP_PuzzleBP", TP_PuzzleBuildSettings);
