Remove the UE4GAME define (nothing checks against it nor should it), and remove unnecessary Rocket overrides in target properties.

diff --git a/Engine/Source/Runtime/Core/Public/Misc/Build.h b/Engine/Source/Runtime/Core/Public/Misc/Build.h
--- a/Engine/Source/Runtime/Core/Public/Misc/Build.h
+++ b/Engine/Source/Runtime/Core/Public/Misc/Build.h
@@ -29,10 +29,6 @@
 #ifndef UE_EDITOR
 	#define UE_EDITOR					0
 #endif
-#ifndef UE4GAME
-	// Only the UE4Game target should ever set this.
-	#define UE4GAME						0
-#endif
 
 /** 
  *   Whether compiling for dedicated server or not.
diff --git a/Engine/Source/UE4Client.Target.cs b/Engine/Source/UE4Client.Target.cs
--- a/Engine/Source/UE4Client.Target.cs
+++ b/Engine/Source/UE4Client.Target.cs
@@ -31,24 +31,6 @@ public class UE4ClientTarget : TargetRules
         ref CPPEnvironmentConfiguration OutCPPEnvironmentConfiguration
         )
     {
-        if (UnrealBuildTool.UnrealBuildTool.BuildingRocket())
-        {
-            UEBuildConfiguration.bCompileLeanAndMeanUE = true;
-
-            // Don't need editor or editor only data
-            UEBuildConfiguration.bBuildEditor = false;
-            UEBuildConfiguration.bBuildWithEditorOnlyData = false;
-
-            UEBuildConfiguration.bCompileAgainstEngine = true;
-
-            // no exports, so no need to verify that a .lib and .exp file was emitted by the linker.
-            OutLinkEnvironmentConfiguration.bHasExports = false;
-        }
-        else
-        {
-            // Tag it as a UE4Game build
-            OutCPPEnvironmentConfiguration.Definitions.Add("UE4GAME=1");
-        }
     }
 
     public override List<UnrealTargetPlatform> GUBP_GetPlatforms_MonolithicOnly(UnrealTargetPlatform HostPlatform)
diff --git a/Engine/Source/UE4Editor.Target.cs b/Engine/Source/UE4Editor.Target.cs
--- a/Engine/Source/UE4Editor.Target.cs
+++ b/Engine/Source/UE4Editor.Target.cs
@@ -44,11 +44,6 @@ public class UE4EditorTarget : TargetRules
 		ref CPPEnvironmentConfiguration OutCPPEnvironmentConfiguration
 		)
 	{
-		if( UnrealBuildTool.UnrealBuildTool.BuildingRocket() )
-		{ 
-			// no exports, so no need to verify that a .lib and .exp file was emitted by the linker.
-			OutLinkEnvironmentConfiguration.bHasExports = false;
-		}
 	}
     public override GUBPProjectOptions GUBP_IncludeProjectInPromotedBuild_EditorTypeOnly(UnrealTargetPlatform HostPlatform)
     {
diff --git a/Engine/Source/UE4Game.Target.cs b/Engine/Source/UE4Game.Target.cs
--- a/Engine/Source/UE4Game.Target.cs
+++ b/Engine/Source/UE4Game.Target.cs
@@ -77,24 +77,6 @@ public class UE4GameTarget : TargetRules
 		ref CPPEnvironmentConfiguration OutCPPEnvironmentConfiguration
 		)
 	{
-		if( UnrealBuildTool.UnrealBuildTool.BuildingRocket() )
-		{
-			UEBuildConfiguration.bCompileLeanAndMeanUE = true;
-
-			// Don't need editor or editor only data
-			UEBuildConfiguration.bBuildEditor = false;
-			UEBuildConfiguration.bBuildWithEditorOnlyData = false;
-
-			UEBuildConfiguration.bCompileAgainstEngine = true;
-
-			// no exports, so no need to verify that a .lib and .exp file was emitted by the linker.
-			OutLinkEnvironmentConfiguration.bHasExports = false;
-		}
-		else
-		{
-			// Tag it as a UE4Game build
-			OutCPPEnvironmentConfiguration.Definitions.Add("UE4GAME=1");
-		}
         if (Target.Platform == UnrealTargetPlatform.IOS)
 		{
 			// to make World Explorers as small as possible we excluded some items from the engine.
diff --git a/Engine/Source/UE4Server.Target.cs b/Engine/Source/UE4Server.Target.cs
--- a/Engine/Source/UE4Server.Target.cs
+++ b/Engine/Source/UE4Server.Target.cs
@@ -30,24 +30,6 @@ public class UE4ServerTarget : TargetRules
         ref CPPEnvironmentConfiguration OutCPPEnvironmentConfiguration
         )
     {
-        if (UnrealBuildTool.UnrealBuildTool.BuildingRocket())
-        {
-            UEBuildConfiguration.bCompileLeanAndMeanUE = true;
-
-            // Don't need editor or editor only data
-            UEBuildConfiguration.bBuildEditor = false;
-            UEBuildConfiguration.bBuildWithEditorOnlyData = false;
-
-            UEBuildConfiguration.bCompileAgainstEngine = true;
-
-            // no exports, so no need to verify that a .lib and .exp file was emitted by the linker.
-            OutLinkEnvironmentConfiguration.bHasExports = false;
-        }
-        else
-        {
-            // Tag it as a UE4Game build
-            OutCPPEnvironmentConfiguration.Definitions.Add("UE4GAME=1");
-        }
     }
 
 	public override bool GetSupportedPlatforms(ref List<UnrealTargetPlatform> OutPlatforms)
