build script, support specific shared packages being moved in front of shared package trigger

diff --git a/Engine/Source/Programs/AutomationTool/Scripts/GUBP.Automation.cs b/Engine/Source/Programs/AutomationTool/Scripts/GUBP.Automation.cs
--- a/Engine/Source/Programs/AutomationTool/Scripts/GUBP.Automation.cs
+++ b/Engine/Source/Programs/AutomationTool/Scripts/GUBP.Automation.cs
@@ -3022,7 +3022,7 @@ public class GUBP : BuildCommand
                         if (!WorkingGameProject.Properties.Targets[GameOrClient].Rules.GUBP_GetConfigs_MonolithicOnly(HostPlatform, Plat).Contains(Config))
                         {
                             throw new AutomationException("Can't make a game/client build for {0} because we didn't build platform {1} config {2}.", WorkingGameProject.GameName, Plat.ToString(), Config.ToString());
-                        }
+                        }						
                     }
                 }
             }
@@ -5917,6 +5917,10 @@ public class GUBP : BuildCommand
                                                 if (PlatPair.TargetPlatform == Plat)
                                                 {
 													var NodeName = AddNode(new FormalBuildNode(this, NonCodeProject, HostPlatform, new List<UnrealTargetPlatform>() { Plat }, new List<UnrealTargetConfiguration>() {PlatPair.TargetConfig}));
+													if(PlatPair.bBeforeTrigger)
+													{
+														RemovePseudodependencyFromNode(FormalBuildNode.StaticGetFullName(NonCodeProject, HostPlatform, new List<UnrealTargetPlatform>() { Plat }, new List<UnrealTargetConfiguration>() { PlatPair.TargetConfig }), WaitForFormalUserInput.StaticGetFullName());
+													}
 													// we don't want this delayed
 													// this would normally wait for the testing phase, we just want to build it right away
 													RemovePseudodependencyFromNode(
@@ -5948,7 +5952,7 @@ public class GUBP : BuildCommand
 													if (PlatPair.bTest)
 													{
 														AddNode(new FormalBuildTestNode(this, NonCodeProject, HostPlatform, Plat, PlatPair.TargetConfig));
-													}																										
+													}													
 												}
 											}
 										}
@@ -6156,6 +6160,10 @@ public class GUBP : BuildCommand
 										}
 										if (FormalNodeName != null)
 										{
+											if (Config.bBeforeTrigger)
+											{
+												RemovePseudodependencyFromNode(FormalBuildNode.StaticGetFullName(CodeProj, HostPlatform, new List<UnrealTargetPlatform>() { Config.TargetPlatform }, new List<UnrealTargetConfiguration>() { Config.TargetConfig }), WaitForFormalUserInput.StaticGetFullName());
+											}
 											// we don't want this delayed
 											// this would normally wait for the testing phase, we just want to build it right away
 											RemovePseudodependencyFromNode(
diff --git a/Engine/Source/Programs/UnrealBuildTool/System/RulesCompiler.cs b/Engine/Source/Programs/UnrealBuildTool/System/RulesCompiler.cs
--- a/Engine/Source/Programs/UnrealBuildTool/System/RulesCompiler.cs
+++ b/Engine/Source/Programs/UnrealBuildTool/System/RulesCompiler.cs
@@ -744,11 +744,13 @@ namespace UnrealBuildTool
             public UnrealTargetPlatform TargetPlatform = UnrealTargetPlatform.Unknown;
             public UnrealTargetConfiguration TargetConfig = UnrealTargetConfiguration.Unknown;
             public bool bTest = false;
-            public GUBPFormalBuild(UnrealTargetPlatform InTargetPlatform, UnrealTargetConfiguration InTargetConfig, bool bInTest = false)
+			public bool bBeforeTrigger = false;
+            public GUBPFormalBuild(UnrealTargetPlatform InTargetPlatform, UnrealTargetConfiguration InTargetConfig, bool bInTest = false, bool bInBeforeTrigger = false)
             {
                 TargetPlatform = InTargetPlatform;
                 TargetConfig = InTargetConfig;
                 bTest = bInTest;
+				bBeforeTrigger = bInBeforeTrigger;
             }
         }
         /// <summary>
@@ -772,7 +774,7 @@ namespace UnrealBuildTool
             public bool bTestWithShared = false;
             public bool bIsMassive = false;
             public bool bCustomWorkflowForPromotion = false;
-			public bool bIsNonCode = false;
+			public bool bIsNonCode = false;			
         }
         public virtual GUBPProjectOptions GUBP_IncludeProjectInPromotedBuild_EditorTypeOnly(UnrealTargetPlatform HostPlatform)
         {
diff --git a/Engine/Source/UE4Editor.Target.cs b/Engine/Source/UE4Editor.Target.cs
--- a/Engine/Source/UE4Editor.Target.cs
+++ b/Engine/Source/UE4Editor.Target.cs
@@ -258,8 +258,8 @@ public class UE4EditorTarget : TargetRules
                     new GUBPFormalBuild(UnrealTargetPlatform.Win32, UnrealTargetConfiguration.Test, false),
                     new GUBPFormalBuild(UnrealTargetPlatform.Win64, UnrealTargetConfiguration.Test, false),
                     new GUBPFormalBuild(UnrealTargetPlatform.Mac, UnrealTargetConfiguration.Test, false),
-                    new GUBPFormalBuild(UnrealTargetPlatform.Android, UnrealTargetConfiguration.Test, false),
-                    new GUBPFormalBuild(UnrealTargetPlatform.IOS, UnrealTargetConfiguration.Test, false),
+                    new GUBPFormalBuild(UnrealTargetPlatform.Android, UnrealTargetConfiguration.Test, false, true),
+                    new GUBPFormalBuild(UnrealTargetPlatform.IOS, UnrealTargetConfiguration.Test, false, true),
                     new GUBPFormalBuild(UnrealTargetPlatform.HTML5, UnrealTargetConfiguration.Test, false),
             };
 
@@ -280,9 +280,9 @@ public class UE4EditorTarget : TargetRules
 
         var TappyChickenBuildSettings = new List<GUBPFormalBuild>
             {
-                    new GUBPFormalBuild(UnrealTargetPlatform.Android, UnrealTargetConfiguration.Shipping, true),
-                    new GUBPFormalBuild(UnrealTargetPlatform.Android, UnrealTargetConfiguration.Test, true),                    
-                    new GUBPFormalBuild(UnrealTargetPlatform.Android, UnrealTargetConfiguration.Development, true),        			
+                    new GUBPFormalBuild(UnrealTargetPlatform.Android, UnrealTargetConfiguration.Shipping, false),
+                    new GUBPFormalBuild(UnrealTargetPlatform.Android, UnrealTargetConfiguration.Test, false, true),                    
+                    new GUBPFormalBuild(UnrealTargetPlatform.Android, UnrealTargetConfiguration.Development, false),        			
 					new GUBPFormalBuild(UnrealTargetPlatform.HTML5, UnrealTargetConfiguration.Shipping, false),
 					new GUBPFormalBuild(UnrealTargetPlatform.HTML5, UnrealTargetConfiguration.Test, false),
             };
@@ -302,16 +302,6 @@ public class UE4EditorTarget : TargetRules
                     new GUBPFormalBuild(UnrealTargetPlatform.XboxOne, UnrealTargetConfiguration.Test, false),
                     new GUBPFormalBuild(UnrealTargetPlatform.PS4, UnrealTargetConfiguration.Test, false),
             };
-		var TP_2DSideScrollerBuildSettings = new List<GUBPFormalBuild>
-            {                    
-                    new GUBPFormalBuild(UnrealTargetPlatform.Android, UnrealTargetConfiguration.Test, false),
-                    new GUBPFormalBuild(UnrealTargetPlatform.HTML5, UnrealTargetConfiguration.Test, false),
-                    new GUBPFormalBuild(UnrealTargetPlatform.Win32, UnrealTargetConfiguration.Test, false),
-                    new GUBPFormalBuild(UnrealTargetPlatform.Win64, UnrealTargetConfiguration.Test, false),
-                    new GUBPFormalBuild(UnrealTargetPlatform.Mac, UnrealTargetConfiguration.Test, false),
-                    new GUBPFormalBuild(UnrealTargetPlatform.XboxOne, UnrealTargetConfiguration.Test, false),
-                    new GUBPFormalBuild(UnrealTargetPlatform.PS4, UnrealTargetConfiguration.Test, false),
-            };
 
         //TP_PUZZLE
         var TP_PuzzleBuildSettings = new List<GUBPFormalBuild>
@@ -370,7 +360,7 @@ public class UE4EditorTarget : TargetRules
 
         //Add Templates to the list with its corresponding settings
         NonCodeProjectNames.Add("FP_FirstPersonBP", AllBuildSettings);
-		NonCodeProjectNames.Add("TP_2DSideScrollerBP", TP_2DSideScrollerBuildSettings);
+		NonCodeProjectNames.Add("TP_2DSideScrollerBP", AllBuildSettings);
         NonCodeProjectNames.Add("TP_FirstPersonBP", AllBuildSettings);
         NonCodeProjectNames.Add("TP_FlyingBP", AllBuildSettings);
         NonCodeProjectNames.Add("TP_PuzzleBP", TP_PuzzleBuildSettings);
