Add the InputDevice module to be precompiled with the launcher build (prompted by PR #940)

diff --git a/Engine/Source/UE4Game.Target.cs b/Engine/Source/UE4Game.Target.cs
--- a/Engine/Source/UE4Game.Target.cs
+++ b/Engine/Source/UE4Game.Target.cs
@@ -56,6 +56,7 @@ public class UE4GameTarget : TargetRules
 	{
 		// Add all the precompiled modules for this target
 		ModuleNames.Add("Launch");
+		ModuleNames.Add("InputDevice");
 		ModuleNames.Add("GameMenuBuilder");
 		if (Target.Platform == UnrealTargetPlatform.Win32 || Target.Platform == UnrealTargetPlatform.Win64)
 		{
