Add XmlParser as an optional runtime module.

diff --git a/Engine/Source/UE4Game.Target.cs b/Engine/Source/UE4Game.Target.cs
--- a/Engine/Source/UE4Game.Target.cs
+++ b/Engine/Source/UE4Game.Target.cs
@@ -58,6 +58,7 @@ public class UE4GameTarget : TargetRules
 		ModuleNames.Add("Launch");
 		ModuleNames.Add("InputDevice");
 		ModuleNames.Add("GameMenuBuilder");
+		ModuleNames.Add("XmlParser");
 		if (Target.Platform == UnrealTargetPlatform.Win32 || Target.Platform == UnrealTargetPlatform.Win64)
 		{
 			ModuleNames.Add("OnlineSubsystemNull");
