fix for UE4-MetalRHI.a not being in the release binary
UE-15934
#ios

diff --git a/Engine/Source/UE4Game.Target.cs b/Engine/Source/UE4Game.Target.cs
--- a/Engine/Source/UE4Game.Target.cs
+++ b/Engine/Source/UE4Game.Target.cs
@@ -87,6 +87,7 @@ public class UE4GameTarget : TargetRules
 			ModuleNames.Add("OnlineSubsystemFacebook");
 			ModuleNames.Add("OnlineSubsystemIOS");
 			ModuleNames.Add("IOSAdvertising");
+			ModuleNames.Add("MetalRHI");
 		}
 		else if (Target.Platform == UnrealTargetPlatform.Android)
 		{
