Fixing UBT by removing unreferenced variable introduced by CL 2561288

Without this fix, UBT fails with error:

/UE4/Engine/Source/UE4Editor.Target.cs(243,13) : warning CS0219: The variable `AllBuildSettings' is assigned but its value is never used

ERROR: UnrealBuildTool encountered an error while compiling source files

diff --git a/Engine/Source/UE4Editor.Target.cs b/Engine/Source/UE4Editor.Target.cs
--- a/Engine/Source/UE4Editor.Target.cs
+++ b/Engine/Source/UE4Editor.Target.cs
@@ -237,22 +237,6 @@ public class UE4EditorTarget : TargetRules
             };
         #endregion
 
-        #region Template Configurations
-        // Template Build Configurations
-        // DEFAULT 
-        var AllBuildSettings = new List<GUBPFormalBuild>
-            {                                     
-                    new GUBPFormalBuild(UnrealTargetPlatform.IOS, UnrealTargetConfiguration.Test, false),
-                    new GUBPFormalBuild(UnrealTargetPlatform.Android, UnrealTargetConfiguration.Test, false),
-                    new GUBPFormalBuild(UnrealTargetPlatform.HTML5, UnrealTargetConfiguration.Test, false),
-                    new GUBPFormalBuild(UnrealTargetPlatform.Win32, UnrealTargetConfiguration.Test, false),
-                    new GUBPFormalBuild(UnrealTargetPlatform.Win64, UnrealTargetConfiguration.Test, false),
-                    new GUBPFormalBuild(UnrealTargetPlatform.Mac, UnrealTargetConfiguration.Test, false),
-                    new GUBPFormalBuild(UnrealTargetPlatform.XboxOne, UnrealTargetConfiguration.Test, false),
-                    new GUBPFormalBuild(UnrealTargetPlatform.PS4, UnrealTargetConfiguration.Test, false),
-            };
-        #endregion
-
         var NonCodeProjectNames = new Dictionary<string, List<GUBPFormalBuild>>();
 
         //Add Samples to the list with its corresponding settings
