UE-16034: Add GameplayAbilities runtime module to the Rocket build, since the editor module is already included.

diff --git a/Engine/Source/UE4Game.Target.cs b/Engine/Source/UE4Game.Target.cs
--- a/Engine/Source/UE4Game.Target.cs
+++ b/Engine/Source/UE4Game.Target.cs
@@ -63,6 +63,7 @@ public class UE4GameTarget : TargetRules
 		ModuleNames.Add("Launch");
 		ModuleNames.Add("InputDevice");
 		ModuleNames.Add("GameMenuBuilder");
+		ModuleNames.Add("GameplayAbilities");
 		ModuleNames.Add("XmlParser");
 		if (Target.Platform == UnrealTargetPlatform.Win32 || Target.Platform == UnrealTargetPlatform.Win64)
 		{
