Create a precompiled UE4Game library so we can link plugins in binary releases.

diff --git a/Engine/Source/UE4Game.Target.cs b/Engine/Source/UE4Game.Target.cs
--- a/Engine/Source/UE4Game.Target.cs
+++ b/Engine/Source/UE4Game.Target.cs
@@ -65,6 +65,7 @@ public class UE4GameTarget : TargetRules
 		ModuleNames.Add("GameMenuBuilder");
 		ModuleNames.Add("GameplayAbilities");
 		ModuleNames.Add("XmlParser");
+		ModuleNames.Add("UE4Game");
 		if (Target.Platform == UnrealTargetPlatform.Win32 || Target.Platform == UnrealTargetPlatform.Win64)
 		{
 			ModuleNames.Add("OnlineSubsystemNull");
