[INTEGRATE] Change 2566718 by [REDACTED_PERSON][REDACTED_HANDLE] on 2015/05/27 11:42:54

	Fixed issue when launching QAGame from a Rocket build due to libs not being pre compiled.

diff --git a/Engine/Source/UE4Game.Target.cs b/Engine/Source/UE4Game.Target.cs
--- a/Engine/Source/UE4Game.Target.cs
+++ b/Engine/Source/UE4Game.Target.cs
@@ -66,6 +66,8 @@ public class UE4GameTarget : TargetRules
 		ModuleNames.Add("GameplayAbilities");
 		ModuleNames.Add("XmlParser");
 		ModuleNames.Add("UE4Game");
+		ModuleNames.Add("AITestSuite");
+		ModuleNames.Add("GameplayDebugger");
 		if (Target.Platform == UnrealTargetPlatform.Win32 || Target.Platform == UnrealTargetPlatform.Win64)
 		{
 			ModuleNames.Add("OnlineSubsystemNull");
