PR #1371: UBT: Add GetModulesToPrecompile() for targets UE4Server and UE4Client (Contributed by [REDACTED_PERSON])

diff --git a/Engine/Source/UE4Client.Target.cs b/Engine/Source/UE4Client.Target.cs
--- a/Engine/Source/UE4Client.Target.cs
+++ b/Engine/Source/UE4Client.Target.cs
@@ -38,6 +38,53 @@ public class UE4ClientTarget : TargetRules
 		return true;
 	}
 
+	public override void GetModulesToPrecompile(TargetInfo Target, List<string> ModuleNames)
+	{
+		// Add all the precompiled modules for this target
+		ModuleNames.Add("Launch");
+		ModuleNames.Add("InputDevice");
+		ModuleNames.Add("GameMenuBuilder");
+		ModuleNames.Add("GameplayAbilities");
+		ModuleNames.Add("XmlParser");
+		ModuleNames.Add("UE4Game");
+		ModuleNames.Add("AITestSuite");
+		ModuleNames.Add("GameplayDebugger");
+		if (Target.Platform == UnrealTargetPlatform.Win32 || Target.Platform == UnrealTargetPlatform.Win64)
+		{
+			ModuleNames.Add("OnlineSubsystemNull");
+			ModuleNames.Add("OnlineSubsystemAmazon");
+			if (UEBuildConfiguration.bCompileSteamOSS == true)
+			{
+				ModuleNames.Add("OnlineSubsystemSteam");
+			}
+			ModuleNames.Add("OnlineSubsystemFacebook");
+		}
+		else if (Target.Platform == UnrealTargetPlatform.Mac || Target.Platform == UnrealTargetPlatform.Linux)
+		{
+			ModuleNames.Add("OnlineSubsystemNull");
+			if (UEBuildConfiguration.bCompileSteamOSS == true)
+			{
+				ModuleNames.Add("OnlineSubsystemSteam");
+			}
+		}
+		else if (Target.Platform == UnrealTargetPlatform.IOS)
+		{
+			ModuleNames.Add("OnlineSubsystemFacebook");
+			ModuleNames.Add("OnlineSubsystemIOS");
+			ModuleNames.Add("IOSAdvertising");
+			ModuleNames.Add("MetalRHI");
+		}
+		else if (Target.Platform == UnrealTargetPlatform.Android)
+		{
+			// @todo android: Add Android online subsystem
+			ModuleNames.Add("AndroidAdvertising");
+		}
+		else if (Target.Platform == UnrealTargetPlatform.HTML5)
+		{
+			ModuleNames.Add("OnlineSubsystemNull");
+		}
+	}
+
     public override List<UnrealTargetPlatform> GUBP_GetPlatforms_MonolithicOnly(UnrealTargetPlatform HostPlatform)
     {
 		List<UnrealTargetPlatform> Platforms = null;
diff --git a/Engine/Source/UE4Server.Target.cs b/Engine/Source/UE4Server.Target.cs
--- a/Engine/Source/UE4Server.Target.cs
+++ b/Engine/Source/UE4Server.Target.cs
@@ -43,6 +43,24 @@ public class UE4ServerTarget : TargetRules
 		return UnrealBuildTool.UnrealBuildTool.GetAllServerPlatforms(ref OutPlatforms, false);
 	}
 
+
+	public override void GetModulesToPrecompile(TargetInfo Target, List<string> ModuleNames)
+	{
+		// Add all the precompiled modules for this target
+		ModuleNames.Add("Launch");
+		ModuleNames.Add("GameplayAbilities");
+		ModuleNames.Add("XmlParser");
+		ModuleNames.Add("GameplayDebugger");
+
+		ModuleNames.Add("OnlineSubsystemNull");
+		ModuleNames.Add("OnlineSubsystemAmazon");
+		if (UEBuildConfiguration.bCompileSteamOSS == true)
+		{
+			ModuleNames.Add("OnlineSubsystemSteam");
+		}
+		ModuleNames.Add("OnlineSubsystemNull");
+	}
+	
     public override List<UnrealTargetPlatform> GUBP_GetPlatforms_MonolithicOnly(UnrealTargetPlatform HostPlatform)
     {
 		List<UnrealTargetPlatform> Platforms = null;
