add RuntimeAssetCache to the precompiled portion of the UE4Game build

diff --git a/Engine/Source/UE4Game.Target.cs b/Engine/Source/UE4Game.Target.cs
--- a/Engine/Source/UE4Game.Target.cs
+++ b/Engine/Source/UE4Game.Target.cs
@@ -68,6 +68,7 @@ public class UE4GameTarget : TargetRules
 		ModuleNames.Add("UE4Game");
 		ModuleNames.Add("AITestSuite");
 		ModuleNames.Add("GameplayDebugger");
+		ModuleNames.Add("RuntimeAssetCache");
 		if (Target.Platform == UnrealTargetPlatform.Win32 || Target.Platform == UnrealTargetPlatform.Win64)
 		{
 			ModuleNames.Add("OnlineSubsystemNull");
