added RuntimeAssetCache to be precompiled for the rest of the UE4 suite (Client, Editor, Server)
UE-19203

diff --git a/Engine/Source/UE4Client.Target.cs b/Engine/Source/UE4Client.Target.cs
--- a/Engine/Source/UE4Client.Target.cs
+++ b/Engine/Source/UE4Client.Target.cs
@@ -49,6 +49,7 @@ public class UE4ClientTarget : TargetRules
 		ModuleNames.Add("UE4Game");
 		ModuleNames.Add("AITestSuite");
 		ModuleNames.Add("GameplayDebugger");
+		ModuleNames.Add("RuntimeAssetCache");
 		if (Target.Platform == UnrealTargetPlatform.Win32 || Target.Platform == UnrealTargetPlatform.Win64)
 		{
 			ModuleNames.Add("OnlineSubsystemNull");
diff --git a/Engine/Source/UE4Editor.Target.cs b/Engine/Source/UE4Editor.Target.cs
--- a/Engine/Source/UE4Editor.Target.cs
+++ b/Engine/Source/UE4Editor.Target.cs
@@ -37,6 +37,7 @@ public class UE4EditorTarget : TargetRules
 	{
 		ModuleNames.Add("Launch");
 		ModuleNames.Add("GameMenuBuilder");
+		ModuleNames.Add("RuntimeAssetCache");
 		if ((Target.Platform == UnrealTargetPlatform.Win32) || (Target.Platform == UnrealTargetPlatform.Win64))
 		{
 			ModuleNames.Add("OnlineSubsystemNull");
diff --git a/Engine/Source/UE4Server.Target.cs b/Engine/Source/UE4Server.Target.cs
--- a/Engine/Source/UE4Server.Target.cs
+++ b/Engine/Source/UE4Server.Target.cs
@@ -51,6 +51,7 @@ public class UE4ServerTarget : TargetRules
 		ModuleNames.Add("GameplayAbilities");
 		ModuleNames.Add("XmlParser");
 		ModuleNames.Add("GameplayDebugger");
+		ModuleNames.Add("RuntimeAssetCache");
 
 		ModuleNames.Add("OnlineSubsystemNull");
 		ModuleNames.Add("OnlineSubsystemAmazon");
