another item needed by QAGame when building under binary
UE-19203

diff --git a/Engine/Source/UE4Game.Target.cs b/Engine/Source/UE4Game.Target.cs
--- a/Engine/Source/UE4Game.Target.cs
+++ b/Engine/Source/UE4Game.Target.cs
@@ -69,6 +69,7 @@ public class UE4GameTarget : TargetRules
 		ModuleNames.Add("AITestSuite");
 		ModuleNames.Add("GameplayDebugger");
 		ModuleNames.Add("RuntimeAssetCache");
+		ModuleNames.Add("UnrealCodeAnalyzerTests");
 		if (Target.Platform == UnrealTargetPlatform.Win32 || Target.Platform == UnrealTargetPlatform.Win64)
 		{
 			ModuleNames.Add("OnlineSubsystemNull");
