one last fix for the build problems with QAGame in binary
UE-19203

diff --git a/Engine/Source/UE4Client.Target.cs b/Engine/Source/UE4Client.Target.cs
--- a/Engine/Source/UE4Client.Target.cs
+++ b/Engine/Source/UE4Client.Target.cs
@@ -50,6 +50,7 @@ public class UE4ClientTarget : TargetRules
 		ModuleNames.Add("AITestSuite");
 		ModuleNames.Add("GameplayDebugger");
 		ModuleNames.Add("RuntimeAssetCache");
+		ModuleNames.Add("UnrealCodeAnalyzerTests");
 		if (Target.Platform == UnrealTargetPlatform.Win32 || Target.Platform == UnrealTargetPlatform.Win64)
 		{
 			ModuleNames.Add("OnlineSubsystemNull");
diff --git a/Engine/Source/UE4Editor.Target.cs b/Engine/Source/UE4Editor.Target.cs
--- a/Engine/Source/UE4Editor.Target.cs
+++ b/Engine/Source/UE4Editor.Target.cs
@@ -38,6 +38,7 @@ public class UE4EditorTarget : TargetRules
 		ModuleNames.Add("Launch");
 		ModuleNames.Add("GameMenuBuilder");
 		ModuleNames.Add("RuntimeAssetCache");
+		ModuleNames.Add("UnrealCodeAnalyzerTests");
 		if ((Target.Platform == UnrealTargetPlatform.Win32) || (Target.Platform == UnrealTargetPlatform.Win64))
 		{
 			ModuleNames.Add("OnlineSubsystemNull");
diff --git a/Engine/Source/UE4Server.Target.cs b/Engine/Source/UE4Server.Target.cs
--- a/Engine/Source/UE4Server.Target.cs
+++ b/Engine/Source/UE4Server.Target.cs
@@ -52,6 +52,7 @@ public class UE4ServerTarget : TargetRules
 		ModuleNames.Add("XmlParser");
 		ModuleNames.Add("GameplayDebugger");
 		ModuleNames.Add("RuntimeAssetCache");
+		ModuleNames.Add("UnrealCodeAnalyzerTests");
 
 		ModuleNames.Add("OnlineSubsystemNull");
 		ModuleNames.Add("OnlineSubsystemAmazon");
