Add JsonUtilities to the list of modules to be built for Rocket.

diff --git a/Engine/Source/UE4Editor.Target.cs b/Engine/Source/UE4Editor.Target.cs
--- a/Engine/Source/UE4Editor.Target.cs
+++ b/Engine/Source/UE4Editor.Target.cs
@@ -37,6 +37,7 @@ public class UE4EditorTarget : TargetRules
 	{
 		ModuleNames.Add("Launch");
 		ModuleNames.Add("GameMenuBuilder");
+		ModuleNames.Add("JsonUtilities");
 		ModuleNames.Add("RuntimeAssetCache");
 		ModuleNames.Add("UnrealCodeAnalyzerTests");
 		if ((Target.Platform == UnrealTargetPlatform.Win32) || (Target.Platform == UnrealTargetPlatform.Win64))
diff --git a/Engine/Source/UE4Game.Target.cs b/Engine/Source/UE4Game.Target.cs
--- a/Engine/Source/UE4Game.Target.cs
+++ b/Engine/Source/UE4Game.Target.cs
@@ -70,6 +70,7 @@ public class UE4GameTarget : TargetRules
 		ModuleNames.Add("GameplayDebugger");
 		ModuleNames.Add("RuntimeAssetCache");
 		ModuleNames.Add("UnrealCodeAnalyzerTests");
+		ModuleNames.Add("JsonUtilities");
 
 		if (Target.Platform == UnrealTargetPlatform.Win32 || Target.Platform == UnrealTargetPlatform.Win64)
 		{
