Removed a name of an internal project from comments

diff --git a/Engine/Source/UE4Game.Target.cs b/Engine/Source/UE4Game.Target.cs
--- a/Engine/Source/UE4Game.Target.cs
+++ b/Engine/Source/UE4Game.Target.cs
@@ -15,7 +15,7 @@ public class UE4GameTarget : TargetRules
 
 		if (Target.Platform == UnrealTargetPlatform.IOS)
 		{
-			// to make World Explorers as small as possible we excluded some items from the engine.
+			// to make iOS projects as small as possible we excluded some items from the engine.
 			// uncomment below to make a smaller iOS build
 			/*bCompileRecast = false;
 			bCompileSpeedTree = false;
