Add a (temporary, hopefully) global.json to work around a more restrictive global.json on some CI systems

#rnx

diff --git a/Engine/Source/global.json b/Engine/Source/global.json
new file mode 100644
--- /dev/null
+++ b/Engine/Source/global.json
@@ -0,0 +1,6 @@
+{
+  "sdk": {
+    "version": "3.1.400",
+    "rollForward": "major"
+  }
+}
\ No newline at end of file
