Remove commented out code for old iOS specific compile options.

diff --git a/Engine/Source/UnrealGame.Target.cs b/Engine/Source/UnrealGame.Target.cs
--- a/Engine/Source/UnrealGame.Target.cs
+++ b/Engine/Source/UnrealGame.Target.cs
@@ -13,18 +13,5 @@ public class UnrealGameTarget : TargetRules
 		BuildEnvironment = TargetBuildEnvironment.Shared;
 
 		ExtraModuleNames.Add("UnrealGame");
-
-		if (Target.Platform == UnrealTargetPlatform.IOS)
-		{
-			// to make iOS projects as small as possible we excluded some items from the engine.
-			// uncomment below to make a smaller iOS build
-			/*bCompileRecast = false;
-			bCompileSpeedTree = false;
-			bCompileAPEX = false;
-			bCompileLeanAndMeanUE = true;
-			bCompilePhysXVehicle = false;
-			bCompileFreeType = false;
-			bCompileForSize = true;*/
-		}
 	}
 }
