New UnrealEditorVM Target and regression tests

#rnx

diff --git a/Engine/Source/UnrealEditorVM.Target.cs b/Engine/Source/UnrealEditorVM.Target.cs
new file mode 100644
--- /dev/null
+++ b/Engine/Source/UnrealEditorVM.Target.cs
@@ -0,0 +1,48 @@
+// Copyright Epic Games, Inc. All Rights Reserved.
+
+using UnrealBuildTool;
+using System.Collections.Generic;
+
+public class UnrealEditorVMTarget : UnrealEditorTarget
+{
+	public UnrealEditorVMTarget( TargetInfo Target ) : base(Target)
+	{
+		// Use the new VM (= don't use the BPVM)
+		bUseVerseBPVM = false;
+		bUseAutoRTFMCompiler = true;
+		BuildEnvironment = TargetBuildEnvironment.Unique;
+		bLegalToDistributeBinary = true;
+
+		DisablePlugins.AddRange(
+			new string[] {
+				"ConcertSyncEntity",
+				"EntityActor",
+				"EntityCore",
+				"EntityFramework",
+				"EpicGamesEngine",
+				"EpicGamesEngineEditor",
+				"EpicGamesEngineRestricted",
+				"EpicGamesTemporary",
+				"UnrealEngineExperimental",
+				"VerseColors",
+				"VerseEngine",
+				"VerseEngineRestricted",
+				"VerseExperimental",
+				"VerseGameplay",
+				"VerseGameplayDebug",
+				"VerseGameplayTags",
+				"VerseGeoScript",
+				"VerseInput",
+				"VerseMovement",
+				"VersePrint",
+				"VerseRestricted",
+				"VerseSimulation",
+				"VerseSimulationMetadata",
+				"VerseUEStandalone",
+				"VerseUI",
+				"VerseWorldPartition",
+				"VGameplayRst",
+			}
+		);
+	}
+}
